Disabled Mass storage parity tests when there's no concurrent storage (i.e. MASS_CONCURRENT_RESERVE == 0)

diff --git a/Engine/Source/Developer/MassEntityTestSuite/Private/MassEntityStorageTest.cpp b/Engine/Source/Developer/MassEntityTestSuite/Private/MassEntityStorageTest.cpp
--- a/Engine/Source/Developer/MassEntityTestSuite/Private/MassEntityStorageTest.cpp
+++ b/Engine/Source/Developer/MassEntityTestSuite/Private/MassEntityStorageTest.cpp
@@ -3,6 +3,9 @@
 #include "AITestsCommon.h"
 #include "MassEntityManager.h"
 
+// not point in "parity testing" if there's only sequential Mass storage available
+#if MASS_CONCURRENT_RESERVE
+
 #define LOCTEXT_NAMESPACE "MassTest"
 
 UE_DISABLE_OPTIMIZATION_SHIP
@@ -387,3 +390,5 @@ IMPLEMENT_AI_INSTANT_TEST(FEntityStorageTest_ConcurrentDataLayoutAssumptions, "S
 UE_ENABLE_OPTIMIZATION_SHIP
 
 #undef LOCTEXT_NAMESPACE
+
+#endif // MASS_CONCURRENT_RESERVE
\ No newline at end of file
