Fix CIS Issue 540807

diff --git a/Engine/Source/Developer/MassEntityTestSuite/Private/MassObserverTest.cpp b/Engine/Source/Developer/MassEntityTestSuite/Private/MassObserverTest.cpp
--- a/Engine/Source/Developer/MassEntityTestSuite/Private/MassObserverTest.cpp
+++ b/Engine/Source/Developer/MassEntityTestSuite/Private/MassObserverTest.cpp
@@ -491,6 +491,8 @@ struct FFragmentObserverTest_EntityCreation_Individual : FFragmentTestBase
 			EntityManager->DestroyEntity(Entity);
 		}
 
+		ObserverProcessor->ForEachEntityChunkExecutionFunction = nullptr;
+
 		return true;
 	}
 };
