Fix CIS Issue 540807

diff --git a/Engine/Source/Developer/MassEntityTestSuite/Private/MassObserverTest.cpp b/Engine/Source/Developer/MassEntityTestSuite/Private/MassObserverTest.cpp
--- a/Engine/Source/Developer/MassEntityTestSuite/Private/MassObserverTest.cpp
+++ b/Engine/Source/Developer/MassEntityTestSuite/Private/MassObserverTest.cpp
@@ -461,7 +461,7 @@ struct FFragmentObserverTest_EntityCreation_Individual : FFragmentTestBase
 		constexpr float TestValue = 123.456f;
 		float ValueOnNotification = 0.f;
 
-		ObserverProcessor->ForEachEntityChunkExecutionFunction = [&ValueOnNotification](FMassExecutionContext& Context)
+		ObserverProcessor->ForEachEntityChunkExecutionFunction = [&ValueOnNotification](FMassExecutionContext& Context) //-V1047 - This lambda is cleared before routine exit
 			{
 				const TConstArrayView<FFragmentStruct> Fragments = Context.GetFragmentView<FFragmentStruct>();
 				for (int32 EntityIndex = 0; EntityIndex < Context.GetNumEntities(); EntityIndex++)
