[Mass] Made types comments format in Mass code uniform

#rnx

diff --git a/Engine/Plugins/AI/MassAI/Source/MassNavigationEditor/Private/MassNavigationTestingActor.cpp b/Engine/Plugins/AI/MassAI/Source/MassNavigationEditor/Private/MassNavigationTestingActor.cpp
--- a/Engine/Plugins/AI/MassAI/Source/MassNavigationEditor/Private/MassNavigationTestingActor.cpp
+++ b/Engine/Plugins/AI/MassAI/Source/MassNavigationEditor/Private/MassNavigationTestingActor.cpp
@@ -10,8 +10,9 @@
 
 #if UE_ENABLE_DEBUG_DRAWING
 
-//////////////////////////////////////////////////////////////////////////
+//-----------------------------------------------------------------------------
 // FMassNavigationTestingSceneProxy
+//-----------------------------------------------------------------------------
 FMassNavigationTestingSceneProxy::FMassNavigationTestingSceneProxy(const UPrimitiveComponent& InComponent)
 	: FDebugRenderSceneProxy(&InComponent)
 {
@@ -41,9 +42,9 @@ uint32 FMassNavigationTestingSceneProxy::GetMemoryFootprint(void) const
 
 #endif // UE_ENABLE_DEBUG_DRAWING
 
-//////////////////////////////////////////////////////////////////////////
+//-----------------------------------------------------------------------------
 // UMassNavigationTestingComponent
-
+//-----------------------------------------------------------------------------
 UMassNavigationTestingComponent::UMassNavigationTestingComponent(const FObjectInitializer& ObjectInitialize)
 	: Super(ObjectInitialize)
 {
@@ -288,9 +289,9 @@ FDebugRenderSceneProxy* UMassNavigationTestingComponent::CreateDebugSceneProxy()
 }
 #endif // UE_ENABLE_DEBUG_DRAWING
 
-//////////////////////////////////////////////////////////////////////////
+//-----------------------------------------------------------------------------
 // UMassNavigationTestingComponent
-
+//-----------------------------------------------------------------------------
 AMassNavigationTestingActor::AMassNavigationTestingActor(const FObjectInitializer& ObjectInitializer)
 	: Super(ObjectInitializer)
 {
diff --git a/Engine/Plugins/Runtime/MassGameplay/Source/MassActors/Public/Translators/MassTranslators_BehaviorTree.h b/Engine/Plugins/Runtime/MassGameplay/Source/MassActors/Public/Translators/MassTranslators_BehaviorTree.h
--- a/Engine/Plugins/Runtime/MassGameplay/Source/MassActors/Public/Translators/MassTranslators_BehaviorTree.h
+++ b/Engine/Plugins/Runtime/MassGameplay/Source/MassActors/Public/Translators/MassTranslators_BehaviorTree.h
@@ -6,7 +6,7 @@
 #include "MassCommonFragments.h"
 #include "MassTranslators_BehaviorTree.generated.h"
 
-//////////////////////////////////////////////////////////////////////////
+
 class UBehaviorTreeComponent;
 
 USTRUCT()
diff --git a/Engine/Plugins/Runtime/MassGameplay/Source/MassReplication/Public/MassReplicationTransformHandlers.h b/Engine/Plugins/Runtime/MassGameplay/Source/MassReplication/Public/MassReplicationTransformHandlers.h
--- a/Engine/Plugins/Runtime/MassGameplay/Source/MassReplication/Public/MassReplicationTransformHandlers.h
+++ b/Engine/Plugins/Runtime/MassGameplay/Source/MassReplication/Public/MassReplicationTransformHandlers.h
@@ -18,8 +18,9 @@ namespace UE::Mass::Replication
 	constexpr float YawReplicateTolerance = 0.004363323f;
 }; //namespace UE::Mass::Replication
 
-
-//////////////////////////////////////////////////////////////////////////// FReplicatedAgentPositionYawData ////////////////////////////////////////////////////////////////////////////
+//-----------------------------------------------------------------------------
+// FReplicatedAgentPositionYawData
+//-----------------------------------------------------------------------------
 /**
  * To replicate position and yaw make a member variable of this class in your FReplicatedAgentBase derived class. In the FReplicatedAgentBase derived class you must also provide an accessor function
  * FReplicatedAgentPathData& GetReplicatedPositionYawDataMutable().
@@ -50,7 +51,10 @@ private:
 	float Yaw = 0;
 };
 
-//////////////////////////////////////////////////////////////////////////// TMassClientBubbleTransformHandler ////////////////////////////////////////////////////////////////////////////
+
+//-----------------------------------------------------------------------------
+// TMassClientBubbleTransformHandler
+//-----------------------------------------------------------------------------
 /**
  * To replicate Transforms make a member variable of this class in your TClientBubbleHandlerBase derived class. This class is a friend of TMassClientBubblePathHandler.
  * It is meant to have access to the protected data declared there.
@@ -192,8 +196,9 @@ void TMassClientBubbleTransformHandler<AgentArrayItem>::SetEntityData(FTransform
 }
 #endif // UE_REPLICATION_COMPILE_CLIENT_CODE
 
-//////////////////////////////////////////////////////////////////////////// FMassReplicationProcessorTransformHandlerBase ////////////////////////////////////////////////////////////////////////////
-
+//-----------------------------------------------------------------------------
+// FMassReplicationProcessorTransformHandlerBase
+//----------------------------------------------------------------------------- 
 class MASSREPLICATION_API FMassReplicationProcessorTransformHandlerBase
 {
 public:
@@ -204,7 +209,9 @@ protected:
 	TArrayView<FTransformFragment> TransformList;
 };
 
-//////////////////////////////////////////////////////////////////////////// FMassReplicationProcessorPositionYawHandler ////////////////////////////////////////////////////////////////////////////
+//-----------------------------------------------------------------------------
+// FMassReplicationProcessorPositionYawHandler
+//----------------------------------------------------------------------------- 
 /**
  * Used to replicate position and yaw by your UMassReplicationProcessorBase derived class. This class should only get used on the server.
  * @todo add #if UE_REPLICATION_COMPILE_SERVER_CODE
diff --git a/Engine/Source/Developer/MassEntityTestSuite/Private/MassEntityTestFarmPlot.cpp b/Engine/Source/Developer/MassEntityTestSuite/Private/MassEntityTestFarmPlot.cpp
--- a/Engine/Source/Developer/MassEntityTestSuite/Private/MassEntityTestFarmPlot.cpp
+++ b/Engine/Source/Developer/MassEntityTestSuite/Private/MassEntityTestFarmPlot.cpp
@@ -12,8 +12,9 @@
 
 //@TODO: Can add a ReadyToHarvest tag Fragment on when things are ready to harvest, to stop them ticking and signal that we need to create an icon
 
-//////////////////////////////////////////////////////////////////////////
+//-----------------------------------------------------------------------------
 // UFarmProcessorBase
+//-----------------------------------------------------------------------------
 UFarmProcessorBase::UFarmProcessorBase()
 	: EntityQuery(*this)
 {
@@ -21,9 +22,9 @@ UFarmProcessorBase::UFarmProcessorBase()
 	bAutoRegisterWithProcessingPhases = false;
 }
 
-//////////////////////////////////////////////////////////////////////////
+//-----------------------------------------------------------------------------
 // UFarmWaterProcessor
-
+//-----------------------------------------------------------------------------
 void UFarmWaterProcessor::ConfigureQueries()
 {
 	EntityQuery.AddRequirement<FFarmWaterFragment>(EMassFragmentAccess::ReadWrite);
@@ -44,9 +45,9 @@ void UFarmWaterProcessor::Execute(FMassEntityManager& EntityManager, FMassExecut
 	});
 }
 
-//////////////////////////////////////////////////////////////////////////
+//-----------------------------------------------------------------------------
 // UFarmHarvestTimerSystem_Flowers
-
+//-----------------------------------------------------------------------------
 void UFarmHarvestTimerSystem_Flowers::ConfigureQueries()
 {
 	EntityQuery.AddRequirement<FHarvestTimerFragment>(EMassFragmentAccess::ReadWrite);
@@ -80,9 +81,9 @@ void UFarmHarvestTimerSystem_Flowers::Execute(FMassEntityManager& EntityManager,
 	});
 }
 
-//////////////////////////////////////////////////////////////////////////
+//-----------------------------------------------------------------------------
 // UFarmHarvestTimerSystem_Crops
-
+//-----------------------------------------------------------------------------
 void UFarmHarvestTimerSystem_Crops::ConfigureQueries()
 {
 	EntityQuery.AddRequirement<FHarvestTimerFragment>(EMassFragmentAccess::ReadWrite);
@@ -109,9 +110,9 @@ void UFarmHarvestTimerSystem_Crops::Execute(FMassEntityManager & EntityManager,
 	});
 }
 
-//////////////////////////////////////////////////////////////////////////
+//-----------------------------------------------------------------------------
 // UFarmHarvestTimerExpired
-
+//-----------------------------------------------------------------------------
 void UFarmHarvestTimerExpired::ConfigureQueries()
 {
 	EntityQuery.AddRequirement<FHarvestTimerFragment>(EMassFragmentAccess::ReadOnly);
@@ -137,9 +138,9 @@ void UFarmHarvestTimerExpired::Execute(FMassEntityManager& EntityManager, FMassE
 	});
 }
 
-//////////////////////////////////////////////////////////////////////////
+//-----------------------------------------------------------------------------
 // UFarmHarvestTimerSetIcon
-
+//-----------------------------------------------------------------------------
 void UFarmHarvestTimerSetIcon::ConfigureQueries()
 {
 	EntityQuery.AddRequirement<FFarmGridCellData>(EMassFragmentAccess::ReadOnly);
@@ -178,9 +179,9 @@ void UFarmHarvestTimerSetIcon::Execute(FMassEntityManager& EntityManager, FMassE
 	});
 }
 
-//////////////////////////////////////////////////////////////////////////
+//-----------------------------------------------------------------------------
 // ALWFragmentTestFarmPlot
-
+//-----------------------------------------------------------------------------
 AMassEntityTestFarmPlot::AMassEntityTestFarmPlot()
 	: SharedEntityManager(MakeShareable(new FMassEntityManager(this)))
 {
diff --git a/Engine/Source/Developer/MassEntityTestSuite/Private/MassSystemRequirementTest.cpp b/Engine/Source/Developer/MassEntityTestSuite/Private/MassSystemRequirementTest.cpp
--- a/Engine/Source/Developer/MassEntityTestSuite/Private/MassSystemRequirementTest.cpp
+++ b/Engine/Source/Developer/MassEntityTestSuite/Private/MassSystemRequirementTest.cpp
@@ -11,9 +11,9 @@
 
 UE_DISABLE_OPTIMIZATION_SHIP
 
-
-//////////////////////////////////////////////////////////////////////////
+//-----------------------------------------------------------------------------
 // Tests
+//----------------------------------------------------------------------------- 
 namespace FMassSystemRequirementTest
 {
 
@@ -40,9 +40,10 @@ struct FSystemRequirementTestBase : FEntityTestBase
 	}
 };	
 
-//////////////////////////////////////////////////////////////////////////
+//-----------------------------------------------------------------------------
 // The main point of this specific test is to verify that we're getting the same answers from both
 // FMassEntityQuery and FMassSubsystemAccess.
+//----------------------------------------------------------------------------- 
 struct FGenericSubsystemAccessAPI : FSystemRequirementTestBase
 {
 	virtual bool InstantTest() override
diff --git a/Engine/Source/Editor/MassEntityEditor/Private/MassEntityEditorSubsystem.cpp b/Engine/Source/Editor/MassEntityEditor/Private/MassEntityEditorSubsystem.cpp
--- a/Engine/Source/Editor/MassEntityEditor/Private/MassEntityEditorSubsystem.cpp
+++ b/Engine/Source/Editor/MassEntityEditor/Private/MassEntityEditorSubsystem.cpp
@@ -64,9 +64,9 @@ struct FMassEditorTickFunction : FTickableEditorObject
 
 } // namespace UE::Mass
 
-//////////////////////////////////////////////////////////////////////
+//-----------------------------------------------------------------------------
 // UMassEntityEditorSubsystem
-
+//-----------------------------------------------------------------------------
 UMassEntityEditorSubsystem::UMassEntityEditorSubsystem()
 	: EntityManager(MakeShareable(new FMassEntityManager(this)))
 	// we explicitly configure FMassProcessingPhaseManager to be in Editor mode since in this context we will not
diff --git a/Engine/Source/Runtime/MassEntity/Private/MassArchetypeData.cpp b/Engine/Source/Runtime/MassEntity/Private/MassArchetypeData.cpp
--- a/Engine/Source/Runtime/MassEntity/Private/MassArchetypeData.cpp
+++ b/Engine/Source/Runtime/MassEntity/Private/MassArchetypeData.cpp
@@ -29,9 +29,9 @@ namespace UE::Mass
 	}
 }
 
-//////////////////////////////////////////////////////////////////////
+//-----------------------------------------------------------------------------
 // FMassArchetypeData
-
+//-----------------------------------------------------------------------------
 FMassArchetypeData::FMassArchetypeData(const FMassArchetypeCreationParams& CreationParams)
 	: NumEntitiesPerChunk(UE::Mass::Private::UninitializedInt32)
 	, EntityListOffsetWithinChunk(UE::Mass::Private::UninitializedInt32)
@@ -1019,9 +1019,9 @@ void FMassArchetypeData::REMOVEME_GetArrayViewForFragmentInChunk(int32 ChunkInde
 	OutNumEntities = Chunk.GetNumInstances();
 }
 
-//////////////////////////////////////////////////////////////////////
+//-----------------------------------------------------------------------------
 // FMassArchetypeData batched api
-
+//-----------------------------------------------------------------------------
 void FMassArchetypeData::BatchAddEntities(TConstArrayView<FMassEntityHandle> Entities, const FMassArchetypeSharedFragmentValues& SharedFragmentValues, TArray<FMassArchetypeEntityCollection::FArchetypeEntityRange>& OutNewRanges)
 {
 	SCOPE_CYCLE_COUNTER(STAT_Mass_ArchetypeBatchAdd);
diff --git a/Engine/Source/Runtime/MassEntity/Private/MassArchetypeTypes.cpp b/Engine/Source/Runtime/MassEntity/Private/MassArchetypeTypes.cpp
--- a/Engine/Source/Runtime/MassEntity/Private/MassArchetypeTypes.cpp
+++ b/Engine/Source/Runtime/MassEntity/Private/MassArchetypeTypes.cpp
@@ -6,17 +6,17 @@
 #include "MassCommandBuffer.h"
 #include "MassEntityUtils.h"
 
-//////////////////////////////////////////////////////////////////////
+//-----------------------------------------------------------------------------
 // FMassArchetypeHandle
-
+//-----------------------------------------------------------------------------
 uint32 GetTypeHash(const FMassArchetypeHandle& Instance)
 {
 	return GetTypeHash(Instance.DataPtr.Get());
 }
 
-//////////////////////////////////////////////////////////////////////
+//-----------------------------------------------------------------------------
 // FMassArchetypeEntityCollection 
-
+//-----------------------------------------------------------------------------
 FMassArchetypeEntityCollection::FMassArchetypeEntityCollection(const FMassArchetypeHandle& InArchetype, TConstArrayView<FMassEntityHandle> InEntities, EDuplicatesHandling DuplicatesHandling)
 	: Archetype(InArchetype)
 {
@@ -188,9 +188,9 @@ bool FMassArchetypeEntityCollection::IsSame(const FMassArchetypeEntityCollection
 	return true;
 }
 
-//////////////////////////////////////////////////////////////////////
+//-----------------------------------------------------------------------------
 // FMassArchetypeEntityCollectionWithPayload
-
+//-----------------------------------------------------------------------------
 void FMassArchetypeEntityCollectionWithPayload::CreateEntityRangesWithPayload(const FMassEntityManager& EntityManager, const TConstArrayView<FMassEntityHandle> Entities
 	, const FMassArchetypeEntityCollection::EDuplicatesHandling DuplicatesHandling, FMassGenericPayloadView Payload
 	, TArray<FMassArchetypeEntityCollectionWithPayload>& OutEntityCollections)
diff --git a/Engine/Source/Runtime/MassEntity/Private/MassCommandBuffer.cpp b/Engine/Source/Runtime/MassEntity/Private/MassCommandBuffer.cpp
--- a/Engine/Source/Runtime/MassEntity/Private/MassCommandBuffer.cpp
+++ b/Engine/Source/Runtime/MassEntity/Private/MassCommandBuffer.cpp
@@ -58,13 +58,14 @@ void GetCommandStatNames(FMassBatchedCommand& Command, FString*& OutName, FAnsiS
 #endif
 } // UE::Mass::Command
 
-//////////////////////////////////////////////////////////////////////
+//-----------------------------------------------------------------------------
 // FMassBatchedCommand
+//-----------------------------------------------------------------------------
 std::atomic<uint32> FMassBatchedCommand::CommandsCounter;
 
-//////////////////////////////////////////////////////////////////////
+//-----------------------------------------------------------------------------
 // FMassCommandBuffer
-
+//-----------------------------------------------------------------------------
 FMassCommandBuffer::FMassCommandBuffer()
 	: OwnerThreadId(FPlatformTLS::GetCurrentThreadId())
 {	
diff --git a/Engine/Source/Runtime/MassEntity/Private/MassEntitySubsystem.cpp b/Engine/Source/Runtime/MassEntity/Private/MassEntitySubsystem.cpp
--- a/Engine/Source/Runtime/MassEntity/Private/MassEntitySubsystem.cpp
+++ b/Engine/Source/Runtime/MassEntity/Private/MassEntitySubsystem.cpp
@@ -35,9 +35,10 @@ namespace UE::Mass::Private
 		};
 	}
 }
-//////////////////////////////////////////////////////////////////////
-// UMassEntitySubsystem
 
+//-----------------------------------------------------------------------------
+// UMassEntitySubsystem
+//-----------------------------------------------------------------------------
 UMassEntitySubsystem::UMassEntitySubsystem()
 	: EntityManager(MakeShareable(new FMassEntityManager(this)))
 {
@@ -90,9 +91,9 @@ void UMassEntitySubsystem::Deinitialize()
 }
 
 #if WITH_MASSENTITY_DEBUG
-//////////////////////////////////////////////////////////////////////
+//-----------------------------------------------------------------------------
 // Debug commands
-
+//-----------------------------------------------------------------------------
 FAutoConsoleCommandWithWorldArgsAndOutputDevice GPrintArchetypesCmd(
 	TEXT("EntityManager.PrintArchetypes"),
 	TEXT("Prints information about all archetypes in the current world"),
diff --git a/Engine/Source/Runtime/MassEntity/Private/MassEntityView.cpp b/Engine/Source/Runtime/MassEntity/Private/MassEntityView.cpp
--- a/Engine/Source/Runtime/MassEntity/Private/MassEntityView.cpp
+++ b/Engine/Source/Runtime/MassEntity/Private/MassEntityView.cpp
@@ -7,8 +7,9 @@
 #include UE_INLINE_GENERATED_CPP_BY_NAME(MassEntityView)
 
 
-//////////////////////////////////////////////////////////////////////
+//-----------------------------------------------------------------------------
 // FMassEntityView
+//-----------------------------------------------------------------------------
 FMassEntityView::FMassEntityView(const FMassArchetypeHandle& ArchetypeHandle, FMassEntityHandle InEntity)
 {
 	Entity = InEntity;
diff --git a/Engine/Source/Runtime/MassEntity/Private/MassExecutionContext.cpp b/Engine/Source/Runtime/MassEntity/Private/MassExecutionContext.cpp
--- a/Engine/Source/Runtime/MassEntity/Private/MassExecutionContext.cpp
+++ b/Engine/Source/Runtime/MassEntity/Private/MassExecutionContext.cpp
@@ -4,9 +4,9 @@
 #include "MassArchetypeData.h"
 
 
-//////////////////////////////////////////////////////////////////////////
+//-----------------------------------------------------------------------------
 // FMassExecutionContext
-
+//-----------------------------------------------------------------------------
 FMassExecutionContext::FMassExecutionContext(FMassEntityManager& InEntityManager, const float InDeltaTimeSeconds, const bool bInFlushDeferredCommands)
 	: SubsystemAccess(InEntityManager.GetWorld())
 	, DeltaTimeSeconds(InDeltaTimeSeconds)
diff --git a/Engine/Source/Runtime/MassEntity/Private/MassRequirements.cpp b/Engine/Source/Runtime/MassEntity/Private/MassRequirements.cpp
--- a/Engine/Source/Runtime/MassEntity/Private/MassRequirements.cpp
+++ b/Engine/Source/Runtime/MassEntity/Private/MassRequirements.cpp
@@ -51,9 +51,9 @@ namespace UE::Mass::Private
 	}
 }
 
-//////////////////////////////////////////////////////////////////////
+//-----------------------------------------------------------------------------
 // FMassSubsystemRequirements
-
+//-----------------------------------------------------------------------------
 void FMassSubsystemRequirements::ExportRequirements(FMassExecutionRequirements& OutRequirements) const
 {
 	OutRequirements.RequiredSubsystems.Read += RequiredConstSubsystems;
@@ -67,9 +67,9 @@ void FMassSubsystemRequirements::Reset()
 	bRequiresGameThreadExecution = false;
 }
 
-//////////////////////////////////////////////////////////////////////
+//-----------------------------------------------------------------------------
 // FMassFragmentRequirements
-
+//-----------------------------------------------------------------------------
 FMassFragmentRequirements::FMassFragmentRequirements(std::initializer_list<UScriptStruct*> InitList)
 {
 	for (const UScriptStruct* FragmentType : InitList)
diff --git a/Engine/Source/Runtime/MassEntity/Private/MassSubsystemAccess.cpp b/Engine/Source/Runtime/MassEntity/Private/MassSubsystemAccess.cpp
--- a/Engine/Source/Runtime/MassEntity/Private/MassSubsystemAccess.cpp
+++ b/Engine/Source/Runtime/MassEntity/Private/MassSubsystemAccess.cpp
@@ -12,9 +12,9 @@ namespace UE::Mass::Private
 	}
 }
 
-//////////////////////////////////////////////////////////////////////////
+//-----------------------------------------------------------------------------
 // FMassSubsystemAccess
-
+//-----------------------------------------------------------------------------
 FMassSubsystemAccess::FMassSubsystemAccess(UWorld* InWorld)
 	: World(InWorld)
 {
diff --git a/Engine/Source/Runtime/MassEntity/Public/MassArchetypeTypes.h b/Engine/Source/Runtime/MassEntity/Public/MassArchetypeTypes.h
--- a/Engine/Source/Runtime/MassEntity/Public/MassArchetypeTypes.h
+++ b/Engine/Source/Runtime/MassEntity/Public/MassArchetypeTypes.h
@@ -22,10 +22,9 @@ struct FMassArchetypeHelper;
 typedef TFunction< void(FMassExecutionContext& /*ExecutionContext*/) > FMassExecuteFunction;
 typedef TFunction< bool(const FMassExecutionContext& /*ExecutionContext*/) > FMassChunkConditionFunction;
 
-
-//////////////////////////////////////////////////////////////////////
+//-----------------------------------------------------------------------------
 // FMassArchetypeHandle
-
+//-----------------------------------------------------------------------------
 // An opaque handle to an archetype
 struct FMassArchetypeHandle final
 {
@@ -46,10 +45,9 @@ private:
 	friend FMassEntityManager;
 };
 
-
-//////////////////////////////////////////////////////////////////////
-// FMassArchetypeEntityCollection 
-
+//-----------------------------------------------------------------------------
+// FMassArchetypeEntityCollection
+//-----------------------------------------------------------------------------
 /** A struct that converts an arbitrary array of entities of given Archetype into a sequence of continuous
  *  entity chunks. The goal is to have the user create an instance of this struct once and run through a bunch of
  *  systems. The runtime code usually uses FMassArchetypeChunkIterator to iterate on the chunk collection.
@@ -156,9 +154,9 @@ private:
 	FMassGenericPayloadViewSlice PayloadSlice;
 };
 
-//////////////////////////////////////////////////////////////////////
+//-----------------------------------------------------------------------------
 // FMassArchetypeChunkIterator
-
+//-----------------------------------------------------------------------------
 /**
  *  The type used to iterate over given archetype's chunks, be it full, continuous chunks or sparse subchunks. It hides
  *  this details from the rest of the system.
@@ -179,9 +177,9 @@ public:
 	const FMassArchetypeEntityCollection::FArchetypeEntityRange& operator*() const { check(bool(*this)); return EntityRanges[CurrentChunkIndex]; }
 };
 
-//////////////////////////////////////////////////////////////////////
+//-----------------------------------------------------------------------------
 // FMassRawEntityInChunkData
-
+//-----------------------------------------------------------------------------
 struct FMassRawEntityInChunkData 
 {
 	FMassRawEntityInChunkData() = default;
@@ -195,9 +193,9 @@ struct FMassRawEntityInChunkData
 	int32 IndexWithinChunk = INDEX_NONE;
 };
 
-//////////////////////////////////////////////////////////////////////
+//-----------------------------------------------------------------------------
 // FMassQueryRequirementIndicesMapping
-
+//-----------------------------------------------------------------------------
 using FMassFragmentIndicesMapping = TArray<int32, TInlineAllocator<16>>;
 
 struct FMassQueryRequirementIndicesMapping
diff --git a/Engine/Source/Runtime/MassEntity/Public/MassCommands.h b/Engine/Source/Runtime/MassEntity/Public/MassCommands.h
--- a/Engine/Source/Runtime/MassEntity/Public/MassCommands.h
+++ b/Engine/Source/Runtime/MassEntity/Public/MassCommands.h
@@ -171,9 +171,9 @@ protected:
 	TArray<FMassEntityHandle> TargetEntities;
 };
 
-//////////////////////////////////////////////////////////////////////
+//-----------------------------------------------------------------------------
 // Entity destruction
-
+//-----------------------------------------------------------------------------
 struct FMassCommandDestroyEntities : public FMassBatchedEntityCommand
 {
 	using Super = FMassBatchedEntityCommand;
@@ -194,9 +194,9 @@ protected:
 	}
 };
 
-//////////////////////////////////////////////////////////////////////
+//-----------------------------------------------------------------------------
 // Simple fragment composition change
-
+//-----------------------------------------------------------------------------
 template<EMassCommandCheckTime CheckTime, typename... TTypes>
 struct FMassCommandAddFragmentsInternal : public FMassBatchedEntityCommand
 {
@@ -243,9 +243,9 @@ protected:
 template<typename... TTypes>
 using FMassCommandRemoveFragments = FMassCommandRemoveFragmentsInternal<EMassCommandCheckTime::CompileTimeCheck, TTypes...>;
 
-//////////////////////////////////////////////////////////////////////
+//-----------------------------------------------------------------------------
 // Simple tag composition change
-
+//-----------------------------------------------------------------------------
 struct FMassCommandChangeTags : public FMassBatchedEntityCommand
 {
 	using Super = FMassBatchedEntityCommand;
@@ -332,9 +332,9 @@ struct FMassCommandSwapTagsInternal : public FMassCommandChangeTags
 template<typename TOld, typename TNew>
 using FMassCommandSwapTags = FMassCommandSwapTagsInternal<EMassCommandCheckTime::CompileTimeCheck, TOld, TNew>;
 
-//////////////////////////////////////////////////////////////////////
+//-----------------------------------------------------------------------------
 // Struct Instances adding and setting
-
+//-----------------------------------------------------------------------------
 template<typename... TOthers>
 struct FMassCommandAddFragmentInstances : public FMassBatchedEntityCommand
 {
@@ -523,9 +523,9 @@ protected:
 	TMap<uint32, FPerSharedFragmentsHashData> Data;
 };
 
-//////////////////////////////////////////////////////////////////////
+//-----------------------------------------------------------------------------
 // Commands that really can't know the types at compile time
-
+//-----------------------------------------------------------------------------
 template<EMassCommandOperationType OpType>
 struct FMassDeferredCommand : public FMassBatchedCommand
 {
diff --git a/Engine/Source/Runtime/MassEntity/Public/MassEntityManagerStorage.h b/Engine/Source/Runtime/MassEntity/Public/MassEntityManagerStorage.h
--- a/Engine/Source/Runtime/MassEntity/Public/MassEntityManagerStorage.h
+++ b/Engine/Source/Runtime/MassEntity/Public/MassEntityManagerStorage.h
@@ -99,9 +99,9 @@ namespace UE::Mass
 	};
 
 	
-	//////////////////////////////////////////////////////////////////////
+	//-----------------------------------------------------------------------------
 	// FSingleThreadedEntityStorage
-
+	//-----------------------------------------------------------------------------
 	/**
 	 * This storage backend should be used when the user of MassEntityManager can guarantee
 	 * that all entity management will be done on a single thread.
@@ -147,9 +147,9 @@ namespace UE::Mass
 		TArray<int32> EntityFreeIndexList;
 	};
 
-	//////////////////////////////////////////////////////////////////////
+	//-----------------------------------------------------------------------------
 	// FConcurrentEntityStorage
-
+	//-----------------------------------------------------------------------------
 	/**
 	 * This storage backend allows for entities to be concurrently reserved.  Reserved entities can also
 	 * be concurrently freed.
diff --git a/Engine/Source/Runtime/MassEntity/Public/MassExecutionContext.h b/Engine/Source/Runtime/MassEntity/Public/MassExecutionContext.h
--- a/Engine/Source/Runtime/MassEntity/Public/MassExecutionContext.h
+++ b/Engine/Source/Runtime/MassEntity/Public/MassExecutionContext.h
@@ -397,8 +397,9 @@ protected:
 	}
 
 public:
-	//////////////////////////////////////////////////////////////////////////
+	//-----------------------------------------------------------------------------
 	// DEPRECATED functions
+	//-----------------------------------------------------------------------------
 
 	template<typename T>
 	UE_DEPRECATED(5.2, "The Get*Subsystem(World) functions have been deprecated. Use the world-less flavor.")
diff --git a/Engine/Source/Runtime/MassEntity/Public/MassSubsystemAccess.h b/Engine/Source/Runtime/MassEntity/Public/MassSubsystemAccess.h
--- a/Engine/Source/Runtime/MassEntity/Public/MassSubsystemAccess.h
+++ b/Engine/Source/Runtime/MassEntity/Public/MassSubsystemAccess.h
@@ -23,9 +23,9 @@ struct MASSENTITY_API FMassSubsystemAccess
 {
 	explicit FMassSubsystemAccess(UWorld* InWorld = nullptr);
 
-	//////////////////////////////////////////////////////////////////////////
+	//-----------------------------------------------------------------------------
 	// Statically-typed subsystems
-
+	//-----------------------------------------------------------------------------
 	template<typename T, typename = typename TEnableIf<TIsDerivedFrom<T, USubsystem>::IsDerived>::Type>
 	T* GetMutableSubsystem()
 	{
@@ -65,9 +65,9 @@ struct MASSENTITY_API FMassSubsystemAccess
 		return *InstancePtr;
 	}
 
-	//////////////////////////////////////////////////////////////////////////
+	//-----------------------------------------------------------------------------
 	// UClass-provided subsystems
-
+	//-----------------------------------------------------------------------------
 	template<typename T, typename = typename TEnableIf<TIsDerivedFrom<T, USubsystem>::IsDerived>::Type>
 	T* GetMutableSubsystem(const TSubclassOf<USubsystem> SubsystemClass)
 	{
@@ -107,9 +107,9 @@ struct MASSENTITY_API FMassSubsystemAccess
 		return *InstancePtr;
 	}
 
-	//////////////////////////////////////////////////////////////////////////
+	//-----------------------------------------------------------------------------
 	// remaining API
-
+	//-----------------------------------------------------------------------------
 	bool CacheSubsystemRequirements(const FMassSubsystemRequirements& SubsystemRequirements);
 	void SetSubsystemRequirements(const FMassSubsystemRequirements& SubsystemRequirements);
 
