[MassEntityManager] added missing AddSharedFragmentToEntity/RemoveSharedFragmentFromEntity (ConstSharedFragments were supported)

diff --git a/Engine/Source/Developer/MassEntityTestSuite/Private/MassSharedFragmentsTest.cpp b/Engine/Source/Developer/MassEntityTestSuite/Private/MassSharedFragmentsTest.cpp
--- a/Engine/Source/Developer/MassEntityTestSuite/Private/MassSharedFragmentsTest.cpp
+++ b/Engine/Source/Developer/MassEntityTestSuite/Private/MassSharedFragmentsTest.cpp
@@ -415,17 +415,17 @@ using FSharedFragment_BatchCreateEntitesWithConstSharedFragment = FSharedFragmen
 IMPLEMENT_AI_INSTANT_TEST(FSharedFragment_BatchCreateEntitesWithConstSharedFragment, "System.Mass.SharedFragments.BatchCreateEntitiesConst");
 
 
-struct FSharedFragment_AddToEntity : FEntityTestBase
+struct FConstSharedFragment_AddToEntity : FEntityTestBase
 {
 	virtual bool InstantTest() override
 	{
 		constexpr int32 TestIntValue = 1023;
-		FTestConstSharedFragment_Int FragmentInstance(TestIntValue);
-		FConstSharedStruct SharedFragmentInstance = FSharedStruct::Make(FragmentInstance);
+		const FTestConstSharedFragment_Int FragmentInstance(TestIntValue);
+		const FConstSharedStruct SharedFragmentInstance = FSharedStruct::Make(FragmentInstance);
 
 		const FMassEntityHandle EntityHandle = EntityManager->CreateEntity(FloatsArchetype);
 
-		FTestConstSharedFragment_Int* EntitySharedFragment = EntityManager->GetConstSharedFragmentDataPtr<FTestConstSharedFragment_Int>(EntityHandle);
+		const FTestConstSharedFragment_Int* EntitySharedFragment = EntityManager->GetConstSharedFragmentDataPtr<FTestConstSharedFragment_Int>(EntityHandle);
 		AITEST_NULL("Initially the entity is not expected to have the shared fragment", EntitySharedFragment);
 
 		EntityManager->AddConstSharedFragmentToEntity(EntityHandle, SharedFragmentInstance);
@@ -438,9 +438,9 @@ struct FSharedFragment_AddToEntity : FEntityTestBase
 		// now we're going to add it again and test the systems behavior, we'll be adding the same FMasSharedFragment type
 		// in both const and non-const way.
 		constexpr int32 DifferentTestIntValue = TestIntValue + 1;
-		FTestConstSharedFragment_Int DifferentFragmentInstance(DifferentTestIntValue);
-		FSharedStruct DifferentSharedFragmentInstance = FSharedStruct::Make(DifferentFragmentInstance);
-		FConstSharedStruct DifferentConstSharedFragmentInstance = FConstSharedStruct::Make(DifferentFragmentInstance);
+		const FTestConstSharedFragment_Int DifferentFragmentInstance(DifferentTestIntValue);
+		const FSharedStruct DifferentSharedFragmentInstance = FSharedStruct::Make(DifferentFragmentInstance);
+		const FConstSharedStruct DifferentConstSharedFragmentInstance = FConstSharedStruct::Make(DifferentFragmentInstance);
 
 		GetTestRunner().AddExpectedError(TEXT("Changing shared fragment value of entities is not supported"), EAutomationExpectedErrorFlags::Contains, 2);
 
@@ -452,19 +452,54 @@ struct FSharedFragment_AddToEntity : FEntityTestBase
 		return true;
 	}
 };
+IMPLEMENT_AI_INSTANT_TEST(FConstSharedFragment_AddToEntity, "System.Mass.ConstSharedFragments.AddToEntity");
+
+struct FSharedFragment_AddToEntity : FEntityTestBase
+{
+	virtual bool InstantTest() override
+	{
+		constexpr int32 TestIntValue = 1023;
+		const FTestSharedFragment_Int FragmentInstance(TestIntValue);
+		const FSharedStruct SharedFragmentInstance = FSharedStruct::Make(FragmentInstance);
+
+		const FMassEntityHandle EntityHandle = EntityManager->CreateEntity(FloatsArchetype);
+
+		const FTestSharedFragment_Int* EntitySharedFragment = EntityManager->GetSharedFragmentDataPtr<FTestSharedFragment_Int>(EntityHandle);
+		AITEST_NULL("Initially the entity is not expected to have the shared fragment", EntitySharedFragment);
+
+		EntityManager->AddSharedFragmentToEntity(EntityHandle, SharedFragmentInstance);
+
+		EntitySharedFragment = EntityManager->GetSharedFragmentDataPtr<FTestSharedFragment_Int>(EntityHandle);
+		AITEST_NOT_NULL("The entity is expected to have the shared fragment after the operation", EntitySharedFragment);
+		AITEST_EQUAL("The the shared fragment is expected to store the configured value", EntitySharedFragment->Value, TestIntValue);
+
+		// at this point the Entity already has a shared fragment of a given type
+		// now we're going to add it again and test the systems behavior.
+		constexpr int32 DifferentTestIntValue = TestIntValue + 1;
+		const FTestSharedFragment_Int DifferentFragmentInstance(DifferentTestIntValue);
+		const FSharedStruct DifferentSharedFragmentInstance = FSharedStruct::Make(DifferentFragmentInstance);
+
+		GetTestRunner().AddExpectedError(TEXT("Changing shared fragment value of entities is not supported"), EAutomationExpectedErrorFlags::Contains, 1);
+
+		const bool bSuccessfullyAddedSharedFragment = EntityManager->AddSharedFragmentToEntity(EntityHandle, DifferentSharedFragmentInstance);
+		AITEST_FALSE("Adding existing shared fragment type should fail", bSuccessfullyAddedSharedFragment);
+
+		return true;
+	}
+};
 IMPLEMENT_AI_INSTANT_TEST(FSharedFragment_AddToEntity, "System.Mass.SharedFragments.AddToEntity");
 
-struct FSharedFragment_RemoveFromEntity : FEntityTestBase
+struct FConstSharedFragment_RemoveFromEntity : FEntityTestBase
 {
 	virtual bool InstantTest() override
 	{
 		constexpr int32 TestIntValue = 1023;
-		FTestConstSharedFragment_Int FragmentInstance(TestIntValue);
-		FConstSharedStruct SharedFragmentInstance = FSharedStruct::Make(FragmentInstance);
+		const FTestConstSharedFragment_Int FragmentInstance(TestIntValue);
+		const FConstSharedStruct SharedFragmentInstance = FSharedStruct::Make(FragmentInstance);
 
 		const FMassEntityHandle EntityHandle = EntityManager->CreateEntity(FloatsArchetype);
 
-		FTestConstSharedFragment_Int* EntitySharedFragment = EntityManager->GetConstSharedFragmentDataPtr<FTestConstSharedFragment_Int>(EntityHandle);
+		const FTestConstSharedFragment_Int* EntitySharedFragment = EntityManager->GetConstSharedFragmentDataPtr<FTestConstSharedFragment_Int>(EntityHandle);
 		AITEST_NULL("Initially the entity is not expected to have the shared fragment", EntitySharedFragment);
 
 		AITEST_FALSE("Attempt to remove shared fragment from entity that doesn't have shared fragment should return false and do nothing",
@@ -489,6 +524,43 @@ struct FSharedFragment_RemoveFromEntity : FEntityTestBase
 		return true;
 	}
 };
+IMPLEMENT_AI_INSTANT_TEST(FConstSharedFragment_RemoveFromEntity, "System.Mass.ConstSharedFragments.RemoveFromEntity");
+
+struct FSharedFragment_RemoveFromEntity : FEntityTestBase
+{
+	virtual bool InstantTest() override
+	{
+		constexpr int32 TestIntValue = 1023;
+		const FTestSharedFragment_Int FragmentInstance(TestIntValue);
+		const FSharedStruct SharedFragmentInstance = FSharedStruct::Make(FragmentInstance);
+
+		const FMassEntityHandle EntityHandle = EntityManager->CreateEntity(FloatsArchetype);
+
+		const FTestSharedFragment_Int* EntitySharedFragment = EntityManager->GetSharedFragmentDataPtr<FTestSharedFragment_Int>(EntityHandle);
+		AITEST_NULL("Initially the entity is not expected to have the shared fragment", EntitySharedFragment);
+
+		AITEST_FALSE("Attempt to remove shared fragment from entity that doesn't have shared fragment should return false and do nothing",
+			EntityManager->RemoveSharedFragmentFromEntity(EntityHandle, *FTestSharedFragment_Int::StaticStruct()));
+
+		AITEST_TRUE("Adding shared fragment to entity should succeed",
+			EntityManager->AddSharedFragmentToEntity(EntityHandle, SharedFragmentInstance));
+
+		EntitySharedFragment = EntityManager->GetSharedFragmentDataPtr<FTestSharedFragment_Int>(EntityHandle);
+		AITEST_NOT_NULL("The entity is expected to have the shared fragment after the operation", EntitySharedFragment);
+		AITEST_EQUAL("The the shared fragment is expected to store the configured value", EntitySharedFragment->Value, TestIntValue);
+
+		AITEST_TRUE(
+			"Removing shared fragment from entity that has the shared fragment should succeed",
+			EntityManager->RemoveSharedFragmentFromEntity(EntityHandle, *FTestSharedFragment_Int::StaticStruct()));
+
+		EntitySharedFragment = EntityManager->GetSharedFragmentDataPtr<FTestSharedFragment_Int>(EntityHandle);
+		AITEST_NULL("The entity is not expected to have the shared fragment after the operation", EntitySharedFragment);
+		
+		AITEST_EQUAL("The the entity's new archetype is the same as the initial one", EntityManager->GetArchetypeForEntity(EntityHandle), FloatsArchetype);
+
+		return true;
+	}
+};
 IMPLEMENT_AI_INSTANT_TEST(FSharedFragment_RemoveFromEntity, "System.Mass.SharedFragments.RemoveFromEntity");
 
 struct FSharedFragment_BatchAddToEntity : FEntityTestBase
diff --git a/Engine/Source/Runtime/MassEntity/Private/MassEntityManager.cpp b/Engine/Source/Runtime/MassEntity/Private/MassEntityManager.cpp
--- a/Engine/Source/Runtime/MassEntity/Private/MassEntityManager.cpp
+++ b/Engine/Source/Runtime/MassEntity/Private/MassEntityManager.cpp
@@ -1355,6 +1355,95 @@ bool FMassEntityManager::RemoveConstSharedFragmentFromEntity(const FMassEntityHa
 	return true;
 }
 
+bool FMassEntityManager::AddSharedFragmentToEntity(const FMassEntityHandle Entity, const FSharedStruct& InSharedFragment)
+{
+	if (!ensureMsgf(InSharedFragment.IsValid(), TEXT("%hs parameter Fragment is expected to be valid"), __FUNCTION__))
+	{
+		return false;
+	}
+
+	CheckIfEntityIsActive(Entity);
+	
+	FMassArchetypeData* CurrentArchetype = GetEntityStorageInterface().GetArchetypeAsShared(Entity.Index).Get();
+	check(CurrentArchetype);
+
+	const UScriptStruct* StructType = InSharedFragment.GetScriptStruct();
+	CA_ASSUME(StructType);
+	if (CurrentArchetype->GetCompositionDescriptor().SharedFragments.Contains(*StructType))
+	{
+		const FMassArchetypeSharedFragmentValues& SharedFragmentValues = CurrentArchetype->GetSharedFragmentValues(Entity);
+		FConstSharedStruct ExistingSharedStruct = SharedFragmentValues.GetSharedFragmentStruct(StructType);
+		if (ExistingSharedStruct == InSharedFragment || ExistingSharedStruct.CompareStructValues(InSharedFragment))
+		{
+			// nothing to do
+			return true;
+		}
+		UE_LOG(LogMass, Warning, TEXT("Changing shared fragment value of entities is not supported"));
+		return false;
+	}
+	
+	FMassArchetypeCompositionDescriptor NewComposition(CurrentArchetype->GetCompositionDescriptor());
+	NewComposition.SharedFragments.Add(*StructType);
+	const FMassArchetypeHandle NewArchetypeHandle = CreateArchetype(NewComposition, FMassArchetypeCreationParams(*CurrentArchetype));
+	check(NewArchetypeHandle.IsValid());
+	FMassArchetypeData* NewArchetype = NewArchetypeHandle.DataPtr.Get();
+	check(NewArchetype);
+
+	const FMassArchetypeSharedFragmentValues& OldSharedFragmentValues = CurrentArchetype->GetSharedFragmentValues(Entity.Index);
+	check(!OldSharedFragmentValues.ContainsType(StructType));
+	FMassArchetypeSharedFragmentValues NewSharedFragmentValues(OldSharedFragmentValues);
+	NewSharedFragmentValues.AddSharedFragment(InSharedFragment);
+	NewSharedFragmentValues.Sort();
+
+	CurrentArchetype->MoveEntityToAnotherArchetype(Entity, *NewArchetype, &NewSharedFragmentValues);
+
+	// Change the entity archetype
+	GetEntityStorageInterface().SetArchetypeFromShared(Entity.Index, NewArchetypeHandle.DataPtr);
+
+	return true;
+}
+
+bool FMassEntityManager::RemoveSharedFragmentFromEntity(const FMassEntityHandle Entity, const UScriptStruct& SharedFragmentType)
+{
+	if (!ensureMsgf(SharedFragmentType.IsChildOf(FMassSharedFragment::StaticStruct()), TEXT("%hs parameter SharedFragmentType is expected to be a FMassSharedFragment"), __FUNCTION__))
+	{
+		return false;
+	}
+	
+	CheckIfEntityIsActive(Entity);
+	
+	FMassArchetypeData* CurrentArchetype = GetEntityStorageInterface().GetArchetypeAsShared(Entity.Index).Get();
+	check(CurrentArchetype);
+	
+	if (!CurrentArchetype->GetCompositionDescriptor().SharedFragments.Contains(SharedFragmentType))
+	{
+		// Nothing to do. Returning false to indicate nothing has been removed, as per function's documentation 
+		return false;
+	}
+
+	FMassArchetypeCompositionDescriptor NewComposition(CurrentArchetype->GetCompositionDescriptor());
+	NewComposition.SharedFragments.Remove(SharedFragmentType);
+	const FMassArchetypeHandle NewArchetypeHandle = CreateArchetype(NewComposition);
+	check(NewArchetypeHandle.IsValid());
+	FMassArchetypeData* NewArchetype = NewArchetypeHandle.DataPtr.Get();
+	check(NewArchetype);
+	
+	const FMassArchetypeSharedFragmentValues& OldSharedFragmentValues = CurrentArchetype->GetSharedFragmentValues(Entity.Index);
+	check(OldSharedFragmentValues.ContainsType(&SharedFragmentType));
+	FMassArchetypeSharedFragmentValues NewSharedFragmentValues(OldSharedFragmentValues);
+	
+	const FMassSharedFragmentBitSet ToRemove(SharedFragmentType);
+	NewSharedFragmentValues.Remove(ToRemove);
+	NewSharedFragmentValues.Sort();
+	
+	CurrentArchetype->MoveEntityToAnotherArchetype(Entity, *NewArchetype, &NewSharedFragmentValues);
+
+	// Change the entity archetype
+	GetEntityStorageInterface().SetArchetypeFromShared(Entity.Index, NewArchetypeHandle.DataPtr);
+	
+	return true;
+}
+
 void FMassEntityManager::BatchChangeTagsForEntities(TConstArrayView<FMassArchetypeEntityCollection> EntityCollections, const FMassTagBitSet& TagsToAdd, const FMassTagBitSet& TagsToRemove)
 {
 	TRACE_CPUPROFILER_EVENT_SCOPE(Mass_BatchChangeTagsForEntities);
diff --git a/Engine/Source/Runtime/MassEntity/Public/MassEntityManager.h b/Engine/Source/Runtime/MassEntity/Public/MassEntityManager.h
--- a/Engine/Source/Runtime/MassEntity/Public/MassEntityManager.h
+++ b/Engine/Source/Runtime/MassEntity/Public/MassEntityManager.h
@@ -379,7 +379,24 @@ public:
 	 */
 	bool RemoveConstSharedFragmentFromEntity(const FMassEntityHandle Entity, const UScriptStruct& ConstSharedFragmentType);
 
-	/** 
+	/**
+	 * Adds a new shared fragment to the given entity. Note that it only works if the given entity doesn't have
+	 * a shared fragment of the given type. The function will give a soft "pass" if the entity has the shared fragment
+	 * of the same value. Setting shared fragment value (i.e. changing) is not supported and the function will log
+	 * a warning if that's attempted.
+	 * @return whether the Entity has the Fragment value assigned to it, regardless of its original state (i.e. the function will
+	 *	return true also if the Entity already had the same values associated with it)
+	 */
+	bool AddSharedFragmentToEntity(const FMassEntityHandle Entity, const FSharedStruct& InSharedFragment);
+
+	/**
+	 * Removes a shared fragment of the given type from the entity.
+	 * Will do nothing if entity did not have the shared fragment.
+	 * @return True if fragment removed from entity, false otherwise.
+	 */
+	bool RemoveSharedFragmentFromEntity(const FMassEntityHandle Entity, const UScriptStruct& SharedFragmentType);
+
+	/**
 	 * Reserves Count number of entities and appends them to InOutEntities
 	 * @return a view into InOutEntities containing only the freshly reserved entities
 	 */
diff --git a/Engine/Source/Runtime/MassEntity/Public/MassEntityTypes.h b/Engine/Source/Runtime/MassEntity/Public/MassEntityTypes.h
--- a/Engine/Source/Runtime/MassEntity/Public/MassEntityTypes.h
+++ b/Engine/Source/Runtime/MassEntity/Public/MassEntityTypes.h
@@ -495,6 +495,12 @@ struct MASSENTITY_API FMassArchetypeSharedFragmentValues
 		return FragmentIndex != INDEX_NONE ? SharedFragments[FragmentIndex] : FSharedStruct();
 	}
 
+	FConstSharedStruct GetSharedFragmentStruct(const UScriptStruct* StructType) const
+	{
+		const int32 FragmentIndex = SharedFragments.IndexOfByPredicate(FStructTypeEqualOperator(StructType));
+		return FragmentIndex != INDEX_NONE ? SharedFragments[FragmentIndex] : FSharedStruct();
+	}
+
 	const FMassSharedFragmentBitSet& GetSharedFragmentBitSet() const
 	{
 		return SharedFragmentBitSet;
