[Mass] added one additional check to System.Mass.SharedFragments.AddToEntity unit test

#rnx

diff --git a/Engine/Source/Developer/MassEntityTestSuite/Private/MassSharedFragmentsTest.cpp b/Engine/Source/Developer/MassEntityTestSuite/Private/MassSharedFragmentsTest.cpp
--- a/Engine/Source/Developer/MassEntityTestSuite/Private/MassSharedFragmentsTest.cpp
+++ b/Engine/Source/Developer/MassEntityTestSuite/Private/MassSharedFragmentsTest.cpp
@@ -433,6 +433,7 @@ struct FConstSharedFragment_AddToEntity : FEntityTestBase
 		EntitySharedFragment = EntityManager->GetConstSharedFragmentDataPtr<FTestConstSharedFragment_Int>(EntityHandle);
 		AITEST_NOT_NULL("The entity is expected to have the shared fragment after the operation", EntitySharedFragment);
 		AITEST_EQUAL("The the shared fragment is expected to store the configured value", EntitySharedFragment->Value, TestIntValue);
+		AITEST_NOT_EQUAL("The the entity's new archetype is not the same as the original one", EntityManager->GetArchetypeForEntity(EntityHandle), FloatsArchetype);
 
 		// at this point the Entity already has a shared fragment of a given type
 		// now we're going to add it again and test the systems behavior, we'll be adding the same FMasSharedFragment type
