[MassEntity] changed MassEntityQuery initialization to require providing a valid EntityManager instance
This requirement led to changes to MassProcessor.ConfigureQueries - it requires a TSharedRef<FMassEntityManager> argument now.

diff --git a/Engine/Plugins/AI/MassAI/Source/MassAIBehavior/Private/MassLookAtProcessors.cpp b/Engine/Plugins/AI/MassAI/Source/MassAIBehavior/Private/MassLookAtProcessors.cpp
--- a/Engine/Plugins/AI/MassAI/Source/MassAIBehavior/Private/MassLookAtProcessors.cpp
+++ b/Engine/Plugins/AI/MassAI/Source/MassAIBehavior/Private/MassLookAtProcessors.cpp
@@ -100,7 +100,7 @@ UMassLookAtProcessor::UMassLookAtProcessor()
 	ExecutionOrder.ExecuteAfter.Add(UE::Mass::ProcessorGroupNames::Representation);
 }
 
-void UMassLookAtProcessor::ConfigureQueries()
+void UMassLookAtProcessor::ConfigureQueries(const TSharedRef<FMassEntityManager>& EntityManager)
 {
 	EntityQuery_Conditional.AddRequirement<FMassLookAtFragment>(EMassFragmentAccess::ReadWrite);
 	EntityQuery_Conditional.AddRequirement<FTransformFragment>(EMassFragmentAccess::ReadOnly);
diff --git a/Engine/Plugins/AI/MassAI/Source/MassAIBehavior/Private/MassStateTreeProcessors.cpp b/Engine/Plugins/AI/MassAI/Source/MassAIBehavior/Private/MassStateTreeProcessors.cpp
--- a/Engine/Plugins/AI/MassAI/Source/MassAIBehavior/Private/MassStateTreeProcessors.cpp
+++ b/Engine/Plugins/AI/MassAI/Source/MassAIBehavior/Private/MassStateTreeProcessors.cpp
@@ -87,7 +87,7 @@ void UMassStateTreeFragmentDestructor::InitializeInternal(UObject& Owner, const
 	Super::InitializeInternal(Owner, EntityManager);
 }
 
-void UMassStateTreeFragmentDestructor::ConfigureQueries()
+void UMassStateTreeFragmentDestructor::ConfigureQueries(const TSharedRef<FMassEntityManager>& EntityManager)
 {
 	EntityQuery.AddRequirement<FMassStateTreeInstanceFragment>(EMassFragmentAccess::ReadWrite);
 	EntityQuery.AddConstSharedRequirement<FMassStateTreeSharedFragment>();
@@ -141,7 +141,7 @@ UMassStateTreeActivationProcessor::UMassStateTreeActivationProcessor()
 	bRequiresGameThreadExecution = true; // due to UMassStateTreeSubsystem RW access
 }
 
-void UMassStateTreeActivationProcessor::ConfigureQueries()
+void UMassStateTreeActivationProcessor::ConfigureQueries(const TSharedRef<FMassEntityManager>& EntityManager)
 {
 	EntityQuery.AddRequirement<FMassStateTreeInstanceFragment>(EMassFragmentAccess::ReadWrite);
 	EntityQuery.AddConstSharedRequirement<FMassStateTreeSharedFragment>();
@@ -266,7 +266,7 @@ void UMassStateTreeProcessor::InitializeInternal(UObject& Owner, const TSharedRe
 	SubscribeToSignal(*SignalSubsystem, UE::Mass::Signals::ContextualAnimTaskFinished);
 }
 
-void UMassStateTreeProcessor::ConfigureQueries()
+void UMassStateTreeProcessor::ConfigureQueries(const TSharedRef<FMassEntityManager>& EntityManager)
 {
 	EntityQuery.AddRequirement<FMassStateTreeInstanceFragment>(EMassFragmentAccess::ReadWrite);
 	EntityQuery.AddConstSharedRequirement<FMassStateTreeSharedFragment>();
diff --git a/Engine/Plugins/AI/MassAI/Source/MassAIBehavior/Private/MassZoneGraphAnnotationProcessors.cpp b/Engine/Plugins/AI/MassAI/Source/MassAIBehavior/Private/MassZoneGraphAnnotationProcessors.cpp
--- a/Engine/Plugins/AI/MassAI/Source/MassAIBehavior/Private/MassZoneGraphAnnotationProcessors.cpp
+++ b/Engine/Plugins/AI/MassAI/Source/MassAIBehavior/Private/MassZoneGraphAnnotationProcessors.cpp
@@ -23,7 +23,7 @@ UMassZoneGraphAnnotationTagsInitializer::UMassZoneGraphAnnotationTagsInitializer
 	Operation = EMassObservedOperation::Add;
 }
 
-void UMassZoneGraphAnnotationTagsInitializer::ConfigureQueries()
+void UMassZoneGraphAnnotationTagsInitializer::ConfigureQueries(const TSharedRef<FMassEntityManager>& EntityManager)
 {
 	EntityQuery.AddRequirement<FMassZoneGraphAnnotationFragment>(EMassFragmentAccess::ReadWrite);
 	EntityQuery.AddRequirement<FMassZoneGraphLaneLocationFragment>(EMassFragmentAccess::ReadOnly);
@@ -78,9 +78,9 @@ void UMassZoneGraphAnnotationTagUpdateProcessor::InitializeInternal(UObject& Own
 	}
 }
 
-void UMassZoneGraphAnnotationTagUpdateProcessor::ConfigureQueries()
+void UMassZoneGraphAnnotationTagUpdateProcessor::ConfigureQueries(const TSharedRef<FMassEntityManager>& EntityManager)
 {
-	Super::ConfigureQueries();
+	Super::ConfigureQueries(EntityManager);
 	EntityQuery.AddRequirement<FMassZoneGraphAnnotationFragment>(EMassFragmentAccess::ReadWrite);
 	EntityQuery.AddRequirement<FMassZoneGraphLaneLocationFragment>(EMassFragmentAccess::ReadOnly);
 	EntityQuery.AddChunkRequirement<FMassZoneGraphAnnotationVariableTickChunkFragment>(EMassFragmentAccess::ReadWrite);
diff --git a/Engine/Plugins/AI/MassAI/Source/MassAIBehavior/Public/MassLookAtProcessors.h b/Engine/Plugins/AI/MassAI/Source/MassAIBehavior/Public/MassLookAtProcessors.h
--- a/Engine/Plugins/AI/MassAI/Source/MassAIBehavior/Public/MassLookAtProcessors.h
+++ b/Engine/Plugins/AI/MassAI/Source/MassAIBehavior/Public/MassLookAtProcessors.h
@@ -29,7 +29,7 @@ class MASSAIBEHAVIOR_API UMassLookAtProcessor : public UMassProcessor
 
 protected:
 
-	virtual void ConfigureQueries() override;
+	virtual void ConfigureQueries(const TSharedRef<FMassEntityManager>& EntityManager) override;
 	virtual void Execute(FMassEntityManager& EntityManager, FMassExecutionContext& Context) override;
 
 	/** Selects a nearby target if possible or use a random fixed direction */
diff --git a/Engine/Plugins/AI/MassAI/Source/MassAIBehavior/Public/MassStateTreeProcessors.h b/Engine/Plugins/AI/MassAI/Source/MassAIBehavior/Public/MassStateTreeProcessors.h
--- a/Engine/Plugins/AI/MassAI/Source/MassAIBehavior/Public/MassStateTreeProcessors.h
+++ b/Engine/Plugins/AI/MassAI/Source/MassAIBehavior/Public/MassStateTreeProcessors.h
@@ -25,7 +25,7 @@ public:
 
 protected:
 	virtual void InitializeInternal(UObject& Owner, const TSharedRef<FMassEntityManager>& EntityManager) override;
-	virtual void ConfigureQueries() override;
+	virtual void ConfigureQueries(const TSharedRef<FMassEntityManager>& EntityManager) override;
 	virtual void Execute(FMassEntityManager& EntityManager, FMassExecutionContext& Context) override;
 
 	FMassEntityQuery EntityQuery;
@@ -51,7 +51,7 @@ class MASSAIBEHAVIOR_API UMassStateTreeActivationProcessor : public UMassProcess
 public:
 	UMassStateTreeActivationProcessor();
 protected:
-	virtual void ConfigureQueries() override;
+	virtual void ConfigureQueries(const TSharedRef<FMassEntityManager>& EntityManager) override;
 	virtual void Execute(FMassEntityManager& EntityManager, FMassExecutionContext& Context) override;
 	
 	FMassEntityQuery EntityQuery;
@@ -70,7 +70,7 @@ public:
 
 protected:
 	virtual void InitializeInternal(UObject& Owner, const TSharedRef<FMassEntityManager>& EntityManager) override;
-	virtual void ConfigureQueries() override;
+	virtual void ConfigureQueries(const TSharedRef<FMassEntityManager>& EntityManager) override;
 	virtual void SignalEntities(FMassEntityManager& EntityManager, FMassExecutionContext& Context, FMassSignalNameLookup& EntitySignals) override;
 
 	/** Configures whether parallel update for FMassArchetypeChunks should be used instead of the default single threaded update (i.e., ParallelForEachEntityChunk instead of ForEachEntityChunk). */
diff --git a/Engine/Plugins/AI/MassAI/Source/MassAIBehavior/Public/MassZoneGraphAnnotationProcessors.h b/Engine/Plugins/AI/MassAI/Source/MassAIBehavior/Public/MassZoneGraphAnnotationProcessors.h
--- a/Engine/Plugins/AI/MassAI/Source/MassAIBehavior/Public/MassZoneGraphAnnotationProcessors.h
+++ b/Engine/Plugins/AI/MassAI/Source/MassAIBehavior/Public/MassZoneGraphAnnotationProcessors.h
@@ -24,7 +24,7 @@ public:
 	UMassZoneGraphAnnotationTagsInitializer();
 
 protected:
-	virtual void ConfigureQueries() override;
+	virtual void ConfigureQueries(const TSharedRef<FMassEntityManager>& EntityManager) override;
 	virtual void Execute(FMassEntityManager& EntityManager, FMassExecutionContext& Context) override;
 
 	FMassEntityQuery EntityQuery;
@@ -43,7 +43,7 @@ public:
 	
 protected:
 	virtual void InitializeInternal(UObject& Owner, const TSharedRef<FMassEntityManager>& EntityManager) override;
-	virtual void ConfigureQueries() override;
+	virtual void ConfigureQueries(const TSharedRef<FMassEntityManager>& EntityManager) override;
 	virtual void Execute(FMassEntityManager& EntityManager, FMassExecutionContext& Context) override;
 
 	void UpdateAnnotationTags(UZoneGraphAnnotationSubsystem& ZoneGraphAnnotationSubsystem, FMassZoneGraphAnnotationFragment& AnnotationTags, const FMassZoneGraphLaneLocationFragment& LaneLocation, FMassEntityHandle Entity);
diff --git a/Engine/Plugins/AI/MassAI/Source/MassAIDebug/Private/GameplayDebuggerCategory_Mass.cpp b/Engine/Plugins/AI/MassAI/Source/MassAIDebug/Private/GameplayDebuggerCategory_Mass.cpp
--- a/Engine/Plugins/AI/MassAI/Source/MassAIDebug/Private/GameplayDebuggerCategory_Mass.cpp
+++ b/Engine/Plugins/AI/MassAI/Source/MassAIDebug/Private/GameplayDebuggerCategory_Mass.cpp
@@ -238,7 +238,7 @@ void FGameplayDebuggerCategory_Mass::PickEntity(const FVector& ViewLocation, con
 		TArray<FMassEntityHandle> Entities;
 		TArray<FVector> Locations;
 		FMassExecutionContext ExecutionContext(EntityManager);
-		FMassEntityQuery Query;
+		FMassEntityQuery Query(EntityManager.AsShared());
 		Query.AddRequirement<FTransformFragment>(EMassFragmentAccess::ReadOnly);
 		Query.ForEachEntityChunk(EntityManager, ExecutionContext, [&Entities, &Locations](FMassExecutionContext& Context)
 		{
@@ -495,7 +495,7 @@ void FGameplayDebuggerCategory_Mass::CollectData(APlayerController* OwnerPC, AAc
 	NearEntityDescriptions.Reset();
 	if (bShowNearEntityOverview && OwnerPC)
 	{
-		FMassEntityQuery EntityQuery;
+		FMassEntityQuery EntityQuery(EntityManager.AsShared());
 		EntityQuery.AddRequirement<FMassStateTreeInstanceFragment>(EMassFragmentAccess::ReadOnly);
 		EntityQuery.AddConstSharedRequirement<FMassStateTreeSharedFragment>();
 		EntityQuery.AddRequirement<FTransformFragment>(EMassFragmentAccess::ReadOnly);
@@ -762,7 +762,7 @@ void FGameplayDebuggerCategory_Mass::CollectData(APlayerController* OwnerPC, AAc
 
 		if (bShowNearEntityAvoidance)
 		{
-			FMassEntityQuery EntityColliderQuery;
+			FMassEntityQuery EntityColliderQuery(EntityManager.AsShared());
 			EntityColliderQuery.AddRequirement<FMassAvoidanceColliderFragment>(EMassFragmentAccess::ReadOnly);
 			EntityColliderQuery.AddRequirement<FTransformFragment>(EMassFragmentAccess::ReadOnly);
 			FMassExecutionContext Context(EntityManager, 0.f);
diff --git a/Engine/Plugins/AI/MassAI/Source/MassAIDebug/Private/MassDebugStateTreeProcessor.cpp b/Engine/Plugins/AI/MassAI/Source/MassAIDebug/Private/MassDebugStateTreeProcessor.cpp
--- a/Engine/Plugins/AI/MassAI/Source/MassAIDebug/Private/MassDebugStateTreeProcessor.cpp
+++ b/Engine/Plugins/AI/MassAI/Source/MassAIDebug/Private/MassDebugStateTreeProcessor.cpp
@@ -22,7 +22,7 @@ UMassDebugStateTreeProcessor::UMassDebugStateTreeProcessor()
 	ExecutionOrder.ExecuteAfter.Add(TEXT("MassStateTreeProcessor"));
 }
 
-void UMassDebugStateTreeProcessor::ConfigureQueries()
+void UMassDebugStateTreeProcessor::ConfigureQueries(const TSharedRef<FMassEntityManager>& EntityManager)
 {
 	EntityQuery.AddRequirement<FMassStateTreeInstanceFragment>(EMassFragmentAccess::ReadOnly);
 	EntityQuery.AddConstSharedRequirement<FMassStateTreeSharedFragment>();
diff --git a/Engine/Plugins/AI/MassAI/Source/MassAIDebug/Public/MassDebugStateTreeProcessor.h b/Engine/Plugins/AI/MassAI/Source/MassAIDebug/Public/MassDebugStateTreeProcessor.h
--- a/Engine/Plugins/AI/MassAI/Source/MassAIDebug/Public/MassDebugStateTreeProcessor.h
+++ b/Engine/Plugins/AI/MassAI/Source/MassAIDebug/Public/MassDebugStateTreeProcessor.h
@@ -17,7 +17,7 @@ class MASSAIDEBUG_API UMassDebugStateTreeProcessor : public UMassProcessor
 protected:
 	UMassDebugStateTreeProcessor();
 
-	virtual void ConfigureQueries() override;
+	virtual void ConfigureQueries(const TSharedRef<FMassEntityManager>& EntityManager) override;
 	virtual void Execute(FMassEntityManager& EntityManager, FMassExecutionContext& Context) override;
 
 	FMassEntityQuery EntityQuery;
diff --git a/Engine/Plugins/AI/MassAI/Source/MassNavigation/Private/MassAvoidanceProcessors.cpp b/Engine/Plugins/AI/MassAI/Source/MassNavigation/Private/MassAvoidanceProcessors.cpp
--- a/Engine/Plugins/AI/MassAI/Source/MassNavigation/Private/MassAvoidanceProcessors.cpp
+++ b/Engine/Plugins/AI/MassAI/Source/MassNavigation/Private/MassAvoidanceProcessors.cpp
@@ -356,7 +356,7 @@ UMassMovingAvoidanceProcessor::UMassMovingAvoidanceProcessor()
 	ExecutionOrder.ExecuteAfter.Add(UE::Mass::ProcessorGroupNames::LOD);
 }
 
-void UMassMovingAvoidanceProcessor::ConfigureQueries()
+void UMassMovingAvoidanceProcessor::ConfigureQueries(const TSharedRef<FMassEntityManager>& EntityManager)
 {
 	EntityQuery.AddRequirement<FMassForceFragment>(EMassFragmentAccess::ReadWrite);
 	EntityQuery.AddRequirement<FMassNavigationEdgesFragment>(EMassFragmentAccess::ReadOnly);
@@ -1021,7 +1021,7 @@ UMassStandingAvoidanceProcessor::UMassStandingAvoidanceProcessor()
 	ExecutionOrder.ExecuteAfter.Add(TEXT("MassMovingAvoidanceProcessor"));
 }
 
-void UMassStandingAvoidanceProcessor::ConfigureQueries()
+void UMassStandingAvoidanceProcessor::ConfigureQueries(const TSharedRef<FMassEntityManager>& EntityManager)
 {
 	EntityQuery.AddRequirement<FMassGhostLocationFragment>(EMassFragmentAccess::ReadWrite);
 	EntityQuery.AddRequirement<FMassNavigationEdgesFragment>(EMassFragmentAccess::ReadOnly);
diff --git a/Engine/Plugins/AI/MassAI/Source/MassNavigation/Private/MassNavigationProcessors.cpp b/Engine/Plugins/AI/MassAI/Source/MassNavigation/Private/MassNavigationProcessors.cpp
--- a/Engine/Plugins/AI/MassAI/Source/MassNavigation/Private/MassNavigationProcessors.cpp
+++ b/Engine/Plugins/AI/MassAI/Source/MassNavigation/Private/MassNavigationProcessors.cpp
@@ -29,7 +29,7 @@ UMassOffLODNavigationProcessor::UMassOffLODNavigationProcessor()
 	ExecutionOrder.ExecuteAfter.Add(UE::Mass::ProcessorGroupNames::Avoidance); // @todo: remove this direct dependency
 }
 
-void UMassOffLODNavigationProcessor::ConfigureQueries()
+void UMassOffLODNavigationProcessor::ConfigureQueries(const TSharedRef<FMassEntityManager>& EntityManager)
 {
 	EntityQuery_Conditional.AddRequirement<FTransformFragment>(EMassFragmentAccess::ReadWrite);
 	EntityQuery_Conditional.AddRequirement<FMassMoveTargetFragment>(EMassFragmentAccess::ReadOnly);
@@ -77,7 +77,7 @@ UMassNavigationSmoothHeightProcessor::UMassNavigationSmoothHeightProcessor()
 	ExecutionOrder.ExecuteAfter.Add(UE::Mass::ProcessorGroupNames::Movement);
 }
 
-void UMassNavigationSmoothHeightProcessor::ConfigureQueries()
+void UMassNavigationSmoothHeightProcessor::ConfigureQueries(const TSharedRef<FMassEntityManager>& EntityManager)
 {
 	EntityQuery.AddRequirement<FTransformFragment>(EMassFragmentAccess::ReadWrite);
 	EntityQuery.AddRequirement<FMassMoveTargetFragment>(EMassFragmentAccess::ReadOnly);
@@ -130,7 +130,7 @@ UMassMoveTargetFragmentInitializer::UMassMoveTargetFragmentInitializer()
 	Operation = EMassObservedOperation::Add;
 }
 
-void UMassMoveTargetFragmentInitializer::ConfigureQueries()
+void UMassMoveTargetFragmentInitializer::ConfigureQueries(const TSharedRef<FMassEntityManager>& EntityManager)
 {
 	InitializerQuery.AddRequirement<FMassMoveTargetFragment>(EMassFragmentAccess::ReadWrite);
 	InitializerQuery.AddRequirement<FTransformFragment>(EMassFragmentAccess::ReadOnly);
@@ -167,9 +167,9 @@ UMassNavigationObstacleGridProcessor::UMassNavigationObstacleGridProcessor()
 	ExecutionOrder.ExecuteAfter.Add(UE::Mass::ProcessorGroupNames::Movement);
 }
 
-void UMassNavigationObstacleGridProcessor::ConfigureQueries()
+void UMassNavigationObstacleGridProcessor::ConfigureQueries(const TSharedRef<FMassEntityManager>& EntityManager)
 {
-	FMassEntityQuery BaseEntityQuery;
+	FMassEntityQuery BaseEntityQuery(EntityManager);
 	BaseEntityQuery.AddRequirement<FTransformFragment>(EMassFragmentAccess::ReadOnly);
 	BaseEntityQuery.AddRequirement<FAgentRadiusFragment>(EMassFragmentAccess::ReadOnly);
 	BaseEntityQuery.AddRequirement<FMassNavigationObstacleGridCellLocationFragment>(EMassFragmentAccess::ReadWrite);
@@ -286,7 +286,7 @@ UMassNavigationObstacleRemoverProcessor::UMassNavigationObstacleRemoverProcessor
 	ExecutionFlags = (int32)EProcessorExecutionFlags::AllNetModes;
 }
 
-void UMassNavigationObstacleRemoverProcessor::ConfigureQueries()
+void UMassNavigationObstacleRemoverProcessor::ConfigureQueries(const TSharedRef<FMassEntityManager>& EntityManager)
 {
 	EntityQuery.AddRequirement<FMassNavigationObstacleGridCellLocationFragment>(EMassFragmentAccess::ReadWrite);
 	EntityQuery.AddSubsystemRequirement<UMassNavigationSubsystem>(EMassFragmentAccess::ReadWrite);
diff --git a/Engine/Plugins/AI/MassAI/Source/MassNavigation/Private/MassSmoothOrientationProcessors.cpp b/Engine/Plugins/AI/MassAI/Source/MassNavigation/Private/MassSmoothOrientationProcessors.cpp
--- a/Engine/Plugins/AI/MassAI/Source/MassNavigation/Private/MassSmoothOrientationProcessors.cpp
+++ b/Engine/Plugins/AI/MassAI/Source/MassNavigation/Private/MassSmoothOrientationProcessors.cpp
@@ -25,7 +25,7 @@ UMassSmoothOrientationProcessor::UMassSmoothOrientationProcessor()
 	ExecutionOrder.ExecuteInGroup = UE::Mass::ProcessorGroupNames::Movement;
 }
 
-void UMassSmoothOrientationProcessor::ConfigureQueries()
+void UMassSmoothOrientationProcessor::ConfigureQueries(const TSharedRef<FMassEntityManager>& EntityManager)
 {
 	HighResEntityQuery.AddRequirement<FMassMoveTargetFragment>(EMassFragmentAccess::ReadOnly);
 	HighResEntityQuery.AddRequirement<FMassVelocityFragment>(EMassFragmentAccess::ReadWrite);
diff --git a/Engine/Plugins/AI/MassAI/Source/MassNavigation/Private/MassSteeringProcessors.cpp b/Engine/Plugins/AI/MassAI/Source/MassNavigation/Private/MassSteeringProcessors.cpp
--- a/Engine/Plugins/AI/MassAI/Source/MassNavigation/Private/MassSteeringProcessors.cpp
+++ b/Engine/Plugins/AI/MassAI/Source/MassNavigation/Private/MassSteeringProcessors.cpp
@@ -59,7 +59,7 @@ UMassSteerToMoveTargetProcessor::UMassSteerToMoveTargetProcessor()
 	ExecutionOrder.ExecuteBefore.Add(UE::Mass::ProcessorGroupNames::Avoidance);
 }
 
-void UMassSteerToMoveTargetProcessor::ConfigureQueries()
+void UMassSteerToMoveTargetProcessor::ConfigureQueries(const TSharedRef<FMassEntityManager>& EntityManager)
 {
 	EntityQuery.AddRequirement<FTransformFragment>(EMassFragmentAccess::ReadOnly);
 	EntityQuery.AddRequirement<FMassMoveTargetFragment>(EMassFragmentAccess::ReadWrite);
diff --git a/Engine/Plugins/AI/MassAI/Source/MassNavigation/Public/Avoidance/MassAvoidanceProcessors.h b/Engine/Plugins/AI/MassAI/Source/MassNavigation/Public/Avoidance/MassAvoidanceProcessors.h
--- a/Engine/Plugins/AI/MassAI/Source/MassNavigation/Public/Avoidance/MassAvoidanceProcessors.h
+++ b/Engine/Plugins/AI/MassAI/Source/MassNavigation/Public/Avoidance/MassAvoidanceProcessors.h
@@ -22,7 +22,7 @@ public:
 	UMassMovingAvoidanceProcessor();
 
 protected:
-	virtual void ConfigureQueries() override;
+	virtual void ConfigureQueries(const TSharedRef<FMassEntityManager>& EntityManager) override;
 	virtual void InitializeInternal(UObject& Owner, const TSharedRef<FMassEntityManager>& EntityManager) override;
 	virtual void Execute(FMassEntityManager& EntityManager, FMassExecutionContext& Context) override;
 
@@ -42,7 +42,7 @@ public:
 	UMassStandingAvoidanceProcessor();
 
 protected:
-	virtual void ConfigureQueries() override;
+	virtual void ConfigureQueries(const TSharedRef<FMassEntityManager>& EntityManager) override;
 	virtual void InitializeInternal(UObject& Owner, const TSharedRef<FMassEntityManager>& EntityManager) override;
 	virtual void Execute(FMassEntityManager& EntityManager, FMassExecutionContext& Context) override;
 
diff --git a/Engine/Plugins/AI/MassAI/Source/MassNavigation/Public/MassNavigationProcessors.h b/Engine/Plugins/AI/MassAI/Source/MassNavigation/Public/MassNavigationProcessors.h
--- a/Engine/Plugins/AI/MassAI/Source/MassNavigation/Public/MassNavigationProcessors.h
+++ b/Engine/Plugins/AI/MassAI/Source/MassNavigation/Public/MassNavigationProcessors.h
@@ -20,7 +20,7 @@ public:
 	UMassOffLODNavigationProcessor();
 
 protected:
-	virtual void ConfigureQueries() override;
+	virtual void ConfigureQueries(const TSharedRef<FMassEntityManager>& EntityManager) override;
 	virtual void Execute(FMassEntityManager& EntityManager, FMassExecutionContext& Context) override;
 
 private:
@@ -40,7 +40,7 @@ public:
 	UMassNavigationSmoothHeightProcessor();
 
 protected:
-	virtual void ConfigureQueries() override;
+	virtual void ConfigureQueries(const TSharedRef<FMassEntityManager>& EntityManager) override;
 	virtual void Execute(FMassEntityManager& EntityManager, FMassExecutionContext& Context) override;
 
 private:
@@ -58,7 +58,7 @@ public:
 	UMassMoveTargetFragmentInitializer();
 	
 protected:
-	virtual void ConfigureQueries() override;
+	virtual void ConfigureQueries(const TSharedRef<FMassEntityManager>& EntityManager) override;
 	virtual void Execute(FMassEntityManager& EntityManager, FMassExecutionContext& Context) override;
 
 	FMassEntityQuery InitializerQuery;
@@ -74,7 +74,7 @@ public:
 	UMassNavigationObstacleGridProcessor();
 
 protected:
-	virtual void ConfigureQueries() override;
+	virtual void ConfigureQueries(const TSharedRef<FMassEntityManager>& EntityManager) override;
 	virtual void Execute(FMassEntityManager& EntityManager, FMassExecutionContext& Context) override;
 
 private:
@@ -92,7 +92,7 @@ class MASSNAVIGATION_API UMassNavigationObstacleRemoverProcessor : public UMassO
 	UMassNavigationObstacleRemoverProcessor();
 
 protected:
-	virtual void ConfigureQueries() override;
+	virtual void ConfigureQueries(const TSharedRef<FMassEntityManager>& EntityManager) override;
 	virtual void Execute(FMassEntityManager& EntityManager, FMassExecutionContext& Context) override;
 
 	FMassEntityQuery EntityQuery;
diff --git a/Engine/Plugins/AI/MassAI/Source/MassNavigation/Public/SmoothOrientation/MassSmoothOrientationProcessors.h b/Engine/Plugins/AI/MassAI/Source/MassNavigation/Public/SmoothOrientation/MassSmoothOrientationProcessors.h
--- a/Engine/Plugins/AI/MassAI/Source/MassNavigation/Public/SmoothOrientation/MassSmoothOrientationProcessors.h
+++ b/Engine/Plugins/AI/MassAI/Source/MassNavigation/Public/SmoothOrientation/MassSmoothOrientationProcessors.h
@@ -17,7 +17,7 @@ public:
 	UMassSmoothOrientationProcessor();
 
 protected:
-	virtual void ConfigureQueries() override;
+	virtual void ConfigureQueries(const TSharedRef<FMassEntityManager>& EntityManager) override;
 	virtual void Execute(FMassEntityManager& EntityManager, FMassExecutionContext& Context) override;
 
 private:
diff --git a/Engine/Plugins/AI/MassAI/Source/MassNavigation/Public/Steering/MassSteeringProcessors.h b/Engine/Plugins/AI/MassAI/Source/MassNavigation/Public/Steering/MassSteeringProcessors.h
--- a/Engine/Plugins/AI/MassAI/Source/MassNavigation/Public/Steering/MassSteeringProcessors.h
+++ b/Engine/Plugins/AI/MassAI/Source/MassNavigation/Public/Steering/MassSteeringProcessors.h
@@ -19,7 +19,7 @@ class MASSNAVIGATION_API UMassSteerToMoveTargetProcessor : public UMassProcessor
 protected:
 	UMassSteerToMoveTargetProcessor();
 	
-	virtual void ConfigureQueries() override;
+	virtual void ConfigureQueries(const TSharedRef<FMassEntityManager>& EntityManager) override;
 	virtual void Execute(FMassEntityManager& EntityManager, FMassExecutionContext& Context) override;
 
 	FMassEntityQuery EntityQuery;
diff --git a/Engine/Plugins/AI/MassAI/Source/MassZoneGraphNavigation/Private/MassZoneGraphNavigationProcessors.cpp b/Engine/Plugins/AI/MassAI/Source/MassZoneGraphNavigation/Private/MassZoneGraphNavigationProcessors.cpp
--- a/Engine/Plugins/AI/MassAI/Source/MassZoneGraphNavigation/Private/MassZoneGraphNavigationProcessors.cpp
+++ b/Engine/Plugins/AI/MassAI/Source/MassZoneGraphNavigation/Private/MassZoneGraphNavigationProcessors.cpp
@@ -46,7 +46,7 @@ UMassZoneGraphLocationInitializer::UMassZoneGraphLocationInitializer()
 	ExecutionFlags = (int32)EProcessorExecutionFlags::AllNetModes;
 }
 
-void UMassZoneGraphLocationInitializer::ConfigureQueries()
+void UMassZoneGraphLocationInitializer::ConfigureQueries(const TSharedRef<FMassEntityManager>& EntityManager)
 {
 	EntityQuery.AddRequirement<FMassZoneGraphLaneLocationFragment>(EMassFragmentAccess::ReadWrite);
 	EntityQuery.AddRequirement<FTransformFragment>(EMassFragmentAccess::ReadOnly);
@@ -127,7 +127,7 @@ void UMassZoneGraphPathFollowProcessor::InitializeInternal(UObject& Owner, const
 	SignalSubsystem = UWorld::GetSubsystem<UMassSignalSubsystem>(Owner.GetWorld());
 }
 
-void UMassZoneGraphPathFollowProcessor::ConfigureQueries()
+void UMassZoneGraphPathFollowProcessor::ConfigureQueries(const TSharedRef<FMassEntityManager>& EntityManager)
 {
 	EntityQuery_Conditional.AddRequirement<FMassZoneGraphShortPathFragment>(EMassFragmentAccess::ReadWrite);
 	EntityQuery_Conditional.AddRequirement<FMassZoneGraphLaneLocationFragment>(EMassFragmentAccess::ReadWrite);
@@ -423,7 +423,7 @@ UMassZoneGraphLaneCacheBoundaryProcessor::UMassZoneGraphLaneCacheBoundaryProcess
 	ExecutionOrder.ExecuteBefore.Add(UE::Mass::ProcessorGroupNames::Avoidance);
 }
 
-void UMassZoneGraphLaneCacheBoundaryProcessor::ConfigureQueries()
+void UMassZoneGraphLaneCacheBoundaryProcessor::ConfigureQueries(const TSharedRef<FMassEntityManager>& EntityManager)
 {
 	EntityQuery.AddRequirement<FMassZoneGraphCachedLaneFragment>(EMassFragmentAccess::ReadOnly);
 	EntityQuery.AddRequirement<FMassMoveTargetFragment>(EMassFragmentAccess::ReadOnly);
diff --git a/Engine/Plugins/AI/MassAI/Source/MassZoneGraphNavigation/Public/MassZoneGraphNavigationProcessors.h b/Engine/Plugins/AI/MassAI/Source/MassZoneGraphNavigation/Public/MassZoneGraphNavigationProcessors.h
--- a/Engine/Plugins/AI/MassAI/Source/MassZoneGraphNavigation/Public/MassZoneGraphNavigationProcessors.h
+++ b/Engine/Plugins/AI/MassAI/Source/MassZoneGraphNavigation/Public/MassZoneGraphNavigationProcessors.h
@@ -21,7 +21,7 @@ public:
 	UMassZoneGraphLocationInitializer();
 
 protected:
-	virtual void ConfigureQueries() override;
+	virtual void ConfigureQueries(const TSharedRef<FMassEntityManager>& EntityManager) override;
 	virtual void Execute(FMassEntityManager& EntityManager, FMassExecutionContext& Context) override;
 
 	FMassEntityQuery EntityQuery;
@@ -39,7 +39,7 @@ protected:
 	UMassZoneGraphPathFollowProcessor();
 	
 	virtual void InitializeInternal(UObject& Owner, const TSharedRef<FMassEntityManager>& EntityManager) override;
-	virtual void ConfigureQueries() override;
+	virtual void ConfigureQueries(const TSharedRef<FMassEntityManager>& EntityManager) override;
 	virtual void Execute(FMassEntityManager& EntityManager, FMassExecutionContext& Context) override;
 
 	FMassEntityQuery EntityQuery_Conditional;
@@ -59,7 +59,7 @@ public:
 	UMassZoneGraphLaneCacheBoundaryProcessor();
 
 protected:
-	virtual void ConfigureQueries() override;
+	virtual void ConfigureQueries(const TSharedRef<FMassEntityManager>& EntityManager) override;
 	virtual void InitializeInternal(UObject& Owner, const TSharedRef<FMassEntityManager>& EntityManager) override;
 	virtual void Execute(FMassEntityManager& EntityManager, FMassExecutionContext& Context) override;
 
diff --git a/Engine/Plugins/AI/MassCrowd/Source/MassCrowd/Private/MassCrowdNavigationProcessor.cpp b/Engine/Plugins/AI/MassCrowd/Source/MassCrowd/Private/MassCrowdNavigationProcessor.cpp
--- a/Engine/Plugins/AI/MassCrowd/Source/MassCrowd/Private/MassCrowdNavigationProcessor.cpp
+++ b/Engine/Plugins/AI/MassCrowd/Source/MassCrowd/Private/MassCrowdNavigationProcessor.cpp
@@ -23,7 +23,7 @@ UMassCrowdLaneTrackingSignalProcessor::UMassCrowdLaneTrackingSignalProcessor()
 	ExecutionOrder.ExecuteBefore.Add(UE::Mass::ProcessorGroupNames::Behavior);
 }
 
-void UMassCrowdLaneTrackingSignalProcessor::ConfigureQueries()
+void UMassCrowdLaneTrackingSignalProcessor::ConfigureQueries(const TSharedRef<FMassEntityManager>& EntityManager)
 {
 	EntityQuery.AddTagRequirement<FMassCrowdTag>(EMassFragmentPresence::All);
 	EntityQuery.AddRequirement<FMassCrowdLaneTrackingFragment>(EMassFragmentAccess::ReadWrite);
@@ -72,7 +72,7 @@ UMassCrowdLaneTrackingDestructor::UMassCrowdLaneTrackingDestructor()
 	Operation = EMassObservedOperation::Remove;
 }
 
-void UMassCrowdLaneTrackingDestructor::ConfigureQueries()
+void UMassCrowdLaneTrackingDestructor::ConfigureQueries(const TSharedRef<FMassEntityManager>& EntityManager)
 {
 	EntityQuery.AddTagRequirement<FMassCrowdTag>(EMassFragmentPresence::All);
 	EntityQuery.AddRequirement<FMassCrowdLaneTrackingFragment>(EMassFragmentAccess::ReadOnly);
@@ -118,7 +118,7 @@ void UMassCrowdDynamicObstacleProcessor::InitializeInternal(UObject& Owner, cons
 	checkf(ZoneGraphAnnotationSubsystem != nullptr, TEXT("UZoneGraphAnnotationSubsystem is mandatory when using this processor."));
 }
 
-void UMassCrowdDynamicObstacleProcessor::ConfigureQueries()
+void UMassCrowdDynamicObstacleProcessor::ConfigureQueries(const TSharedRef<FMassEntityManager>& EntityManager)
 {
 	EntityQuery_Conditional.AddRequirement<FTransformFragment>(EMassFragmentAccess::ReadOnly);
 	EntityQuery_Conditional.AddRequirement<FMassVelocityFragment>(EMassFragmentAccess::ReadOnly, EMassFragmentPresence::Optional);
@@ -234,7 +234,7 @@ UMassCrowdDynamicObstacleInitializer::UMassCrowdDynamicObstacleInitializer()
 	Operation = EMassObservedOperation::Add;
 }
 
-void UMassCrowdDynamicObstacleInitializer::ConfigureQueries()
+void UMassCrowdDynamicObstacleInitializer::ConfigureQueries(const TSharedRef<FMassEntityManager>& EntityManager)
 {
 	EntityQuery.AddRequirement<FTransformFragment>(EMassFragmentAccess::ReadOnly);
 	EntityQuery.AddRequirement<FMassCrowdObstacleFragment>(EMassFragmentAccess::ReadWrite);
@@ -284,7 +284,7 @@ void UMassCrowdDynamicObstacleDeinitializer::InitializeInternal(UObject& Owner,
 		, TEXT("UZoneGraphAnnotationSubsystem is mandatory when using this processor."));
 }
 
-void UMassCrowdDynamicObstacleDeinitializer::ConfigureQueries()
+void UMassCrowdDynamicObstacleDeinitializer::ConfigureQueries(const TSharedRef<FMassEntityManager>& EntityManager)
 {
 	EntityQuery.AddRequirement<FMassCrowdObstacleFragment>(EMassFragmentAccess::ReadWrite);
 }
diff --git a/Engine/Plugins/AI/MassCrowd/Source/MassCrowd/Private/MassCrowdServerRepresentationLODProcessor.cpp b/Engine/Plugins/AI/MassCrowd/Source/MassCrowd/Private/MassCrowdServerRepresentationLODProcessor.cpp
--- a/Engine/Plugins/AI/MassCrowd/Source/MassCrowd/Private/MassCrowdServerRepresentationLODProcessor.cpp
+++ b/Engine/Plugins/AI/MassCrowd/Source/MassCrowd/Private/MassCrowdServerRepresentationLODProcessor.cpp
@@ -33,7 +33,7 @@ UMassCrowdServerRepresentationLODProcessor::UMassCrowdServerRepresentationLODPro
 	LODMaxCount[EMassLOD::Off] = INT_MAX;
 }
 
-void UMassCrowdServerRepresentationLODProcessor::ConfigureQueries()
+void UMassCrowdServerRepresentationLODProcessor::ConfigureQueries(const TSharedRef<FMassEntityManager>& EntityManager)
 {
 	EntityQuery.AddTagRequirement<FMassCrowdTag>(EMassFragmentPresence::All);
 	EntityQuery.AddRequirement<FTransformFragment>(EMassFragmentAccess::ReadOnly);
diff --git a/Engine/Plugins/AI/MassCrowd/Source/MassCrowd/Private/MassCrowdServerRepresentationProcessor.cpp b/Engine/Plugins/AI/MassCrowd/Source/MassCrowd/Private/MassCrowdServerRepresentationProcessor.cpp
--- a/Engine/Plugins/AI/MassCrowd/Source/MassCrowd/Private/MassCrowdServerRepresentationProcessor.cpp
+++ b/Engine/Plugins/AI/MassCrowd/Source/MassCrowd/Private/MassCrowdServerRepresentationProcessor.cpp
@@ -14,8 +14,8 @@ UMassCrowdServerRepresentationProcessor::UMassCrowdServerRepresentationProcessor
 	ExecutionOrder.ExecuteAfter.Add(UE::Mass::ProcessorGroupNames::LOD);
 }
 
-void UMassCrowdServerRepresentationProcessor::ConfigureQueries()
+void UMassCrowdServerRepresentationProcessor::ConfigureQueries(const TSharedRef<FMassEntityManager>& EntityManager)
 {
-	Super::ConfigureQueries();
+	Super::ConfigureQueries(EntityManager);
 	EntityQuery.AddTagRequirement<FMassCrowdTag>(EMassFragmentPresence::All);
 }
diff --git a/Engine/Plugins/AI/MassCrowd/Source/MassCrowd/Private/MassCrowdVisualizationLODProcessor.cpp b/Engine/Plugins/AI/MassCrowd/Source/MassCrowd/Private/MassCrowdVisualizationLODProcessor.cpp
--- a/Engine/Plugins/AI/MassCrowd/Source/MassCrowd/Private/MassCrowdVisualizationLODProcessor.cpp
+++ b/Engine/Plugins/AI/MassCrowd/Source/MassCrowd/Private/MassCrowdVisualizationLODProcessor.cpp
@@ -38,9 +38,9 @@ UMassCrowdVisualizationLODProcessor::UMassCrowdVisualizationLODProcessor()
 	ExecutionOrder.ExecuteAfter.Add(UE::Mass::ProcessorGroupNames::LODCollector);
 }
 
-void UMassCrowdVisualizationLODProcessor::ConfigureQueries()
+void UMassCrowdVisualizationLODProcessor::ConfigureQueries(const TSharedRef<FMassEntityManager>& EntityManager)
 {
-	Super::ConfigureQueries();
+	Super::ConfigureQueries(EntityManager);
 
 	CloseEntityQuery.AddTagRequirement<FMassCrowdTag>(EMassFragmentPresence::All);
 	CloseEntityAdjustDistanceQuery.AddTagRequirement<FMassCrowdTag>(EMassFragmentPresence::All);
@@ -108,9 +108,9 @@ UMassCrowdLODCollectorProcessor::UMassCrowdLODCollectorProcessor()
 	bAutoRegisterWithProcessingPhases = true;
 }
 
-void UMassCrowdLODCollectorProcessor::ConfigureQueries()
+void UMassCrowdLODCollectorProcessor::ConfigureQueries(const TSharedRef<FMassEntityManager>& EntityManager)
 {
-	Super::ConfigureQueries();
+	Super::ConfigureQueries(EntityManager);
 
 	EntityQuery_VisibleRangeAndOnLOD.AddTagRequirement<FMassCrowdTag>(EMassFragmentPresence::All);
 	EntityQuery_VisibleRangeOnly.AddTagRequirement<FMassCrowdTag>(EMassFragmentPresence::All);
diff --git a/Engine/Plugins/AI/MassCrowd/Source/MassCrowd/Private/MassCrowdVisualizationProcessor.cpp b/Engine/Plugins/AI/MassCrowd/Source/MassCrowd/Private/MassCrowdVisualizationProcessor.cpp
--- a/Engine/Plugins/AI/MassCrowd/Source/MassCrowd/Private/MassCrowdVisualizationProcessor.cpp
+++ b/Engine/Plugins/AI/MassCrowd/Source/MassCrowd/Private/MassCrowdVisualizationProcessor.cpp
@@ -34,9 +34,9 @@ UMassCrowdVisualizationProcessor::UMassCrowdVisualizationProcessor()
 	bRequiresGameThreadExecution = true;
 }
 
-void UMassCrowdVisualizationProcessor::ConfigureQueries()
+void UMassCrowdVisualizationProcessor::ConfigureQueries(const TSharedRef<FMassEntityManager>& EntityManager)
 {
-	Super::ConfigureQueries();
+	Super::ConfigureQueries(EntityManager);
 	EntityQuery.AddTagRequirement<FMassCrowdTag>(EMassFragmentPresence::All);
 }
 
@@ -54,7 +54,7 @@ UMassDebugCrowdVisualizationProcessor::UMassDebugCrowdVisualizationProcessor()
 	bRequiresGameThreadExecution = true;
 }
 
-void UMassDebugCrowdVisualizationProcessor::ConfigureQueries()
+void UMassDebugCrowdVisualizationProcessor::ConfigureQueries(const TSharedRef<FMassEntityManager>& EntityManager)
 {
 	EntityQuery.AddTagRequirement<FMassCrowdTag>(EMassFragmentPresence::All);
 
diff --git a/Engine/Plugins/AI/MassCrowd/Source/MassCrowd/Public/MassCrowdNavigationProcessor.h b/Engine/Plugins/AI/MassCrowd/Source/MassCrowd/Public/MassCrowdNavigationProcessor.h
--- a/Engine/Plugins/AI/MassCrowd/Source/MassCrowd/Public/MassCrowdNavigationProcessor.h
+++ b/Engine/Plugins/AI/MassCrowd/Source/MassCrowd/Public/MassCrowdNavigationProcessor.h
@@ -20,7 +20,7 @@ public:
 
 protected:
 	virtual void InitializeInternal(UObject& Owner, const TSharedRef<FMassEntityManager>& EntityManager) override;
-	virtual void ConfigureQueries() override;
+	virtual void ConfigureQueries(const TSharedRef<FMassEntityManager>& EntityManager) override;
 	virtual void SignalEntities(FMassEntityManager& EntityManager, FMassExecutionContext& Context, FMassSignalNameLookup& /*Unused*/) override;
 };
 
@@ -34,7 +34,7 @@ public:
 	UMassCrowdLaneTrackingDestructor();
 
 protected:
-	virtual void ConfigureQueries() override;
+	virtual void ConfigureQueries(const TSharedRef<FMassEntityManager>& EntityManager) override;
 	virtual void Execute(FMassEntityManager& EntityManager, FMassExecutionContext& Context) override;
 
 	FMassEntityQuery EntityQuery;
@@ -51,7 +51,7 @@ public:
 
 protected:
 	virtual void InitializeInternal(UObject& Owner, const TSharedRef<FMassEntityManager>& EntityManager)override;
-	virtual void ConfigureQueries() override;
+	virtual void ConfigureQueries(const TSharedRef<FMassEntityManager>& EntityManager) override;
 	virtual void Execute(FMassEntityManager& EntityManager, FMassExecutionContext& Context);
 
 	void OnStop(FMassCrowdObstacleFragment& OutObstacle, const float Radius);
@@ -73,7 +73,7 @@ public:
 	UMassCrowdDynamicObstacleInitializer();
 
 protected:
-	virtual void ConfigureQueries() override;
+	virtual void ConfigureQueries(const TSharedRef<FMassEntityManager>& EntityManager) override;
 	virtual void Execute(FMassEntityManager& EntityManager, FMassExecutionContext& Context) override;
 
 	FMassEntityQuery EntityQuery;
@@ -90,7 +90,7 @@ public:
 
 protected:
 	virtual void InitializeInternal(UObject& Owner, const TSharedRef<FMassEntityManager>& EntityManager) override;
-	virtual void ConfigureQueries() override;
+	virtual void ConfigureQueries(const TSharedRef<FMassEntityManager>& EntityManager) override;
 	virtual void Execute(FMassEntityManager& EntityManager, FMassExecutionContext& Context) override;
 
 	UPROPERTY(Transient)
diff --git a/Engine/Plugins/AI/MassCrowd/Source/MassCrowd/Public/MassCrowdServerRepresentationLODProcessor.h b/Engine/Plugins/AI/MassCrowd/Source/MassCrowd/Public/MassCrowdServerRepresentationLODProcessor.h
--- a/Engine/Plugins/AI/MassCrowd/Source/MassCrowd/Public/MassCrowdServerRepresentationLODProcessor.h
+++ b/Engine/Plugins/AI/MassCrowd/Source/MassCrowd/Public/MassCrowdServerRepresentationLODProcessor.h
@@ -22,7 +22,7 @@ public:
 
 protected:
 
-	virtual void ConfigureQueries() override;
+	virtual void ConfigureQueries(const TSharedRef<FMassEntityManager>& EntityManager) override;
 	virtual void InitializeInternal(UObject& InOwner, const TSharedRef<FMassEntityManager>& EntityManager) override;
 	virtual void Execute(FMassEntityManager& EntityManager, FMassExecutionContext& Context) override;
 
diff --git a/Engine/Plugins/AI/MassCrowd/Source/MassCrowd/Public/MassCrowdServerRepresentationProcessor.h b/Engine/Plugins/AI/MassCrowd/Source/MassCrowd/Public/MassCrowdServerRepresentationProcessor.h
--- a/Engine/Plugins/AI/MassCrowd/Source/MassCrowd/Public/MassCrowdServerRepresentationProcessor.h
+++ b/Engine/Plugins/AI/MassCrowd/Source/MassCrowd/Public/MassCrowdServerRepresentationProcessor.h
@@ -18,5 +18,5 @@ public:
 	UMassCrowdServerRepresentationProcessor();
 
 	/** Configure the owned FMassEntityQuery instances to express processor's requirements */
-	virtual void ConfigureQueries() override;
+	virtual void ConfigureQueries(const TSharedRef<FMassEntityManager>& EntityManager) override;
 };
\ No newline at end of file
diff --git a/Engine/Plugins/AI/MassCrowd/Source/MassCrowd/Public/MassCrowdVisualizationLODProcessor.h b/Engine/Plugins/AI/MassCrowd/Source/MassCrowd/Public/MassCrowdVisualizationLODProcessor.h
--- a/Engine/Plugins/AI/MassCrowd/Source/MassCrowd/Public/MassCrowdVisualizationLODProcessor.h
+++ b/Engine/Plugins/AI/MassCrowd/Source/MassCrowd/Public/MassCrowdVisualizationLODProcessor.h
@@ -17,7 +17,7 @@ public:
 	UMassCrowdVisualizationLODProcessor();
 
 protected:
-	virtual void ConfigureQueries() override;
+	virtual void ConfigureQueries(const TSharedRef<FMassEntityManager>& EntityManager) override;
 	virtual void Execute(FMassEntityManager& EntityManager, FMassExecutionContext& Context) override;
 };
 
@@ -32,6 +32,6 @@ class MASSCROWD_API UMassCrowdLODCollectorProcessor : public UMassLODCollectorPr
 	UMassCrowdLODCollectorProcessor();
 
 protected:
-	virtual void ConfigureQueries() override;
+	virtual void ConfigureQueries(const TSharedRef<FMassEntityManager>& EntityManager) override;
 };
 
diff --git a/Engine/Plugins/AI/MassCrowd/Source/MassCrowd/Public/MassCrowdVisualizationProcessor.h b/Engine/Plugins/AI/MassCrowd/Source/MassCrowd/Public/MassCrowdVisualizationProcessor.h
--- a/Engine/Plugins/AI/MassCrowd/Source/MassCrowd/Public/MassCrowdVisualizationProcessor.h
+++ b/Engine/Plugins/AI/MassCrowd/Source/MassCrowd/Public/MassCrowdVisualizationProcessor.h
@@ -18,7 +18,7 @@ public:
 protected:
 
 	/** Configure the owned FMassEntityQuery instances to express processor's requirements */
-	virtual void ConfigureQueries() override;
+	virtual void ConfigureQueries(const TSharedRef<FMassEntityManager>& EntityManager) override;
 };
 
 /**
@@ -34,7 +34,7 @@ public:
 protected:
 
 	/** Configure the owned FMassEntityQuery instances to express processor's requirements */
-	virtual void ConfigureQueries() override;
+	virtual void ConfigureQueries(const TSharedRef<FMassEntityManager>& EntityManager) override;
 
 	virtual void InitializeInternal(UObject& Owner, const TSharedRef<FMassEntityManager>& EntityManager) override;
 
diff --git a/Engine/Plugins/Experimental/EditorDataStorage/Source/TedsCore/Private/Processors/TypedElementProcessorAdaptors.cpp b/Engine/Plugins/Experimental/EditorDataStorage/Source/TedsCore/Private/Processors/TypedElementProcessorAdaptors.cpp
--- a/Engine/Plugins/Experimental/EditorDataStorage/Source/TedsCore/Private/Processors/TypedElementProcessorAdaptors.cpp
+++ b/Engine/Plugins/Experimental/EditorDataStorage/Source/TedsCore/Private/Processors/TypedElementProcessorAdaptors.cpp
@@ -1302,7 +1302,7 @@ bool UTypedElementQueryProcessorCallbackAdapterProcessorBase::ConfigureQueryCall
 	return Result;
 }
 
-void UTypedElementQueryProcessorCallbackAdapterProcessorBase::ConfigureQueries()
+void UTypedElementQueryProcessorCallbackAdapterProcessorBase::ConfigureQueries(const TSharedRef<FMassEntityManager>&)
 {
 	// When the extended query information is provided the native query will already be fully configured.
 }
@@ -1445,7 +1445,7 @@ bool UTypedElementQueryObserverCallbackAdapterProcessorBase::ConfigureQueryCallb
 	return Result;
 }
 
-void UTypedElementQueryObserverCallbackAdapterProcessorBase::ConfigureQueries()
+void UTypedElementQueryObserverCallbackAdapterProcessorBase::ConfigureQueries(const TSharedRef<FMassEntityManager>& EntityManager)
 {
 	// When the extended query information is provided the native query will already be fully configured.
 }
diff --git a/Engine/Plugins/Experimental/EditorDataStorage/Source/TedsCore/Private/Processors/TypedElementProcessorAdaptors.h b/Engine/Plugins/Experimental/EditorDataStorage/Source/TedsCore/Private/Processors/TypedElementProcessorAdaptors.h
--- a/Engine/Plugins/Experimental/EditorDataStorage/Source/TedsCore/Private/Processors/TypedElementProcessorAdaptors.h
+++ b/Engine/Plugins/Experimental/EditorDataStorage/Source/TedsCore/Private/Processors/TypedElementProcessorAdaptors.h
@@ -121,7 +121,7 @@ protected:
 		FExtendedQueryStore& QueryStore,
 		FEnvironment& Environment,
 		TArrayView<FMassEntityQuery> Subqueries);
-	void ConfigureQueries() override;
+	void ConfigureQueries(const TSharedRef<FMassEntityManager>&) override;
 	void Execute(FMassEntityManager& EntityManager, FMassExecutionContext& TargetParentQuery) override;
 	
 	void PostInitProperties() override;
@@ -346,7 +346,7 @@ protected:
 		FExtendedQueryStore& QueryStore,
 		FEnvironment& Environment,
 		TArrayView<FMassEntityQuery> Subqueries);
-	virtual void ConfigureQueries() override;
+	virtual void ConfigureQueries(const TSharedRef<FMassEntityManager>& EntityManager) override;
 	virtual void Execute(FMassEntityManager& EntityManager, FMassExecutionContext& TargetParentQuery) override;
 
 	virtual void PostInitProperties() override;
diff --git a/Engine/Plugins/Experimental/EditorDataStorage/Source/TedsCore/Private/Queries/TypedElementExtendedQueryStore.cpp b/Engine/Plugins/Experimental/EditorDataStorage/Source/TedsCore/Private/Queries/TypedElementExtendedQueryStore.cpp
--- a/Engine/Plugins/Experimental/EditorDataStorage/Source/TedsCore/Private/Queries/TypedElementExtendedQueryStore.cpp
+++ b/Engine/Plugins/Experimental/EditorDataStorage/Source/TedsCore/Private/Queries/TypedElementExtendedQueryStore.cpp
@@ -151,8 +151,8 @@ namespace UE::Editor::DataStorage
 		FExtendedQuery& StoredQuery = GetMutableChecked(Result);
 		StoredQuery.Description = MoveTemp(Query);
 
-		FMassEntityQuery& NativeQuery = SetupNativeQuery(StoredQuery.Description, StoredQuery);
-		
+		FMassEntityQuery& NativeQuery = SetupNativeQuery(StoredQuery.Description, StoredQuery, Environment);
+
 		bool bContinueSetup =               SetupDynamicColumns(StoredQuery.Description, Environment);
 		bContinueSetup = bContinueSetup &&	SetupSelectedColumns(StoredQuery.Description, NativeQuery);
 		bContinueSetup = bContinueSetup && SetupChunkFilters(Result, StoredQuery.Description, Environment, NativeQuery);
@@ -555,7 +555,7 @@ namespace UE::Editor::DataStorage
 	}
 
 	FMassEntityQuery& FExtendedQueryStore::SetupNativeQuery(
-		FQueryDescription& Query, FExtendedQuery& StoredQuery)
+		FQueryDescription& Query, FExtendedQuery& StoredQuery, FEnvironment& Environment)
 	{
 		/**
 		 * Mass verifies that queries that are used by processors are on the processor themselves. It does this by taking the address of the query
@@ -608,10 +608,12 @@ namespace UE::Editor::DataStorage
 				default:
 					checkf(false, TEXT("The current Typed Elements Data Storage backend doesn't support %i subqueries per processor query."),
 						Query.Subqueries.Num());
+					StoredQuery.NativeQuery.Initialize(Environment.GetMassEntityManager().AsShared());
 					return StoredQuery.NativeQuery;
 				}
 				StoredQuery.Processor.Reset(Processor);
 				StoredQuery.QueryReference = &Processor->GetQuery();
+				StoredQuery.QueryReference->Initialize(Environment.GetMassEntityManager().AsShared());
 				return Processor->GetQuery();
 			}
 			case EQueryCallbackType::ObserveAdd:
@@ -655,6 +657,7 @@ namespace UE::Editor::DataStorage
 				}
 				StoredQuery.Processor.Reset(Observer);
 				StoredQuery.QueryReference = &Observer->GetQuery();
+				StoredQuery.QueryReference->Initialize(Environment.GetMassEntityManager().AsShared());
 				return Observer->GetQuery();
 			}
 			case EQueryCallbackType::PhasePreparation:
@@ -666,7 +669,8 @@ namespace UE::Editor::DataStorage
 				break;
 			}
 		}
-		
+
+		StoredQuery.NativeQuery.Initialize(Environment.GetMassEntityManager().AsShared());
 		return StoredQuery.NativeQuery;
 	}
 
diff --git a/Engine/Plugins/Experimental/EditorDataStorage/Source/TedsCore/Private/Queries/TypedElementExtendedQueryStore.h b/Engine/Plugins/Experimental/EditorDataStorage/Source/TedsCore/Private/Queries/TypedElementExtendedQueryStore.h
--- a/Engine/Plugins/Experimental/EditorDataStorage/Source/TedsCore/Private/Queries/TypedElementExtendedQueryStore.h
+++ b/Engine/Plugins/Experimental/EditorDataStorage/Source/TedsCore/Private/Queries/TypedElementExtendedQueryStore.h
@@ -178,7 +178,7 @@ namespace UE::Editor::DataStorage
 		EDirectQueryExecutionFlags ExecutionFlags,
 			CallbackReference Callback);
 
-		FMassEntityQuery& SetupNativeQuery(IEditorDataStorageProvider::FQueryDescription& Query, FExtendedQuery& StoredQuery);
+		FMassEntityQuery& SetupNativeQuery(IEditorDataStorageProvider::FQueryDescription& Query, FExtendedQuery& StoredQuery, FEnvironment& Environment);
 		bool SetupDynamicColumns(IEditorDataStorageProvider::FQueryDescription& Query, FEnvironment& Environment);
 		bool SetupSelectedColumns(IEditorDataStorageProvider::FQueryDescription& Query, FMassEntityQuery& NativeQuery);
 		bool SetupConditions(IEditorDataStorageProvider::FQueryDescription& Query, FEnvironment& Environment, FMassEntityQuery& NativeQuery);
diff --git a/Engine/Plugins/Runtime/InstancedActors/Source/InstancedActors/Private/InstancedActorsData.cpp b/Engine/Plugins/Runtime/InstancedActors/Source/InstancedActors/Private/InstancedActorsData.cpp
--- a/Engine/Plugins/Runtime/InstancedActors/Source/InstancedActors/Private/InstancedActorsData.cpp
+++ b/Engine/Plugins/Runtime/InstancedActors/Source/InstancedActors/Private/InstancedActorsData.cpp
@@ -446,7 +446,7 @@ void UInstancedActorsData::DespawnEntities()
 	InstanceTransforms.SetNumZeroed(NumInstances);
 	NumValidInstances = 0;
 
-	FMassEntityQuery InstancedActorLocationQuery;
+	FMassEntityQuery InstancedActorLocationQuery(MassEntityManager.AsShared());
 	InstancedActorLocationQuery.AddRequirement<FInstancedActorsFragment>(EMassFragmentAccess::ReadOnly);
 	InstancedActorLocationQuery.AddRequirement<FTransformFragment>(EMassFragmentAccess::ReadOnly);
 
diff --git a/Engine/Plugins/Runtime/InstancedActors/Source/InstancedActors/Private/InstancedActorsDebugProcessor.cpp b/Engine/Plugins/Runtime/InstancedActors/Source/InstancedActors/Private/InstancedActorsDebugProcessor.cpp
--- a/Engine/Plugins/Runtime/InstancedActors/Source/InstancedActors/Private/InstancedActorsDebugProcessor.cpp
+++ b/Engine/Plugins/Runtime/InstancedActors/Source/InstancedActors/Private/InstancedActorsDebugProcessor.cpp
@@ -116,7 +116,7 @@ namespace UE::InstancedActors::Debug
 				{
 					FMassEntityManager& EntityManager = EntitySubsystem->GetMutableEntityManager();
 
-					FMassEntityQuery DetailedLODEntityQuery;
+					FMassEntityQuery DetailedLODEntityQuery(EntityManager.AsShared());
 					DetailedLODEntityQuery.AddRequirement<FInstancedActorsFragment>(EMassFragmentAccess::None);
 					DetailedLODEntityQuery.AddRequirement<FMassActorFragment>(EMassFragmentAccess::ReadOnly);
 					DetailedLODEntityQuery.AddRequirement<FMassRepresentationFragment>(EMassFragmentAccess::ReadWrite);
@@ -171,7 +171,7 @@ UInstancedActorsDebugProcessor::UInstancedActorsDebugProcessor()
 	bRequiresGameThreadExecution = true; // for debug drawing
 }
 
-void UInstancedActorsDebugProcessor::ConfigureQueries()
+void UInstancedActorsDebugProcessor::ConfigureQueries(const TSharedRef<FMassEntityManager>& EntityManager)
 {
 	DebugAllEntityQuery.AddRequirement<FInstancedActorsFragment>(EMassFragmentAccess::None);
 	DebugAllEntityQuery.AddRequirement<FTransformFragment>(EMassFragmentAccess::ReadOnly);
@@ -308,7 +308,7 @@ UInstancedActorsDebugProcessor::UInstancedActorsDebugProcessor()
 	bAutoRegisterWithProcessingPhases = false;
 }
 
-void UInstancedActorsDebugProcessor::ConfigureQueries()
+void UInstancedActorsDebugProcessor::ConfigureQueries(const TSharedRef<FMassEntityManager>& EntityManager)
 {
 }
 
diff --git a/Engine/Plugins/Runtime/InstancedActors/Source/InstancedActors/Private/InstancedActorsInitializerProcessor.cpp b/Engine/Plugins/Runtime/InstancedActors/Source/InstancedActors/Private/InstancedActorsInitializerProcessor.cpp
--- a/Engine/Plugins/Runtime/InstancedActors/Source/InstancedActors/Private/InstancedActorsInitializerProcessor.cpp
+++ b/Engine/Plugins/Runtime/InstancedActors/Source/InstancedActors/Private/InstancedActorsInitializerProcessor.cpp
@@ -19,7 +19,7 @@ UInstancedActorsInitializerProcessor::UInstancedActorsInitializerProcessor()
 	bAutoRegisterWithProcessingPhases = false;
 }
 
-void UInstancedActorsInitializerProcessor::ConfigureQueries()
+void UInstancedActorsInitializerProcessor::ConfigureQueries(const TSharedRef<FMassEntityManager>& EntityManager)
 {
 	ProcessorRequirements.AddSubsystemRequirement<UMassSignalSubsystem>(EMassFragmentAccess::ReadWrite);
 
diff --git a/Engine/Plugins/Runtime/InstancedActors/Source/InstancedActors/Private/InstancedActorsManager.cpp b/Engine/Plugins/Runtime/InstancedActors/Source/InstancedActors/Private/InstancedActorsManager.cpp
--- a/Engine/Plugins/Runtime/InstancedActors/Source/InstancedActors/Private/InstancedActorsManager.cpp
+++ b/Engine/Plugins/Runtime/InstancedActors/Source/InstancedActors/Private/InstancedActorsManager.cpp
@@ -262,8 +262,9 @@ void AInstancedActorsManager::InitializeModifyAndSpawnEntities()
 {
 	check(IsValid(InstancedActorSubsystem));
 
-	if (InstancedActorLocationQuery.IsEmpty())
+	if (InstancedActorLocationQuery.IsInitialized() == false && ensure(MassEntityManager))
 	{
+		InstancedActorLocationQuery.Initialize(MassEntityManager.ToSharedRef());
 		InstancedActorLocationQuery.AddRequirement<FInstancedActorsFragment>(EMassFragmentAccess::ReadOnly);
 		InstancedActorLocationQuery.AddRequirement<FTransformFragment>(EMassFragmentAccess::ReadOnly);
 	}
diff --git a/Engine/Plugins/Runtime/InstancedActors/Source/InstancedActors/Private/InstancedActorsStationaryLODBatchProcessor.cpp b/Engine/Plugins/Runtime/InstancedActors/Source/InstancedActors/Private/InstancedActorsStationaryLODBatchProcessor.cpp
--- a/Engine/Plugins/Runtime/InstancedActors/Source/InstancedActors/Private/InstancedActorsStationaryLODBatchProcessor.cpp
+++ b/Engine/Plugins/Runtime/InstancedActors/Source/InstancedActors/Private/InstancedActorsStationaryLODBatchProcessor.cpp
@@ -90,7 +90,7 @@ UInstancedActorsStationaryLODBatchProcessor::UInstancedActorsStationaryLODBatchP
 	DelayPerBulkLOD[(int)EInstancedActorsBulkLOD::Off] = 10.0;
 }
 
-void UInstancedActorsStationaryLODBatchProcessor::ConfigureQueries()
+void UInstancedActorsStationaryLODBatchProcessor::ConfigureQueries(const TSharedRef<FMassEntityManager>& EntityManager)
 {
 	ProcessorRequirements.AddSubsystemRequirement<UMassLODSubsystem>(EMassFragmentAccess::ReadOnly);
 	if (HasAnyFlags(RF_ClassDefaultObject) == false)
diff --git a/Engine/Plugins/Runtime/InstancedActors/Source/InstancedActors/Private/InstancedActorsVisualizationProcessor.cpp b/Engine/Plugins/Runtime/InstancedActors/Source/InstancedActors/Private/InstancedActorsVisualizationProcessor.cpp
--- a/Engine/Plugins/Runtime/InstancedActors/Source/InstancedActors/Private/InstancedActorsVisualizationProcessor.cpp
+++ b/Engine/Plugins/Runtime/InstancedActors/Source/InstancedActors/Private/InstancedActorsVisualizationProcessor.cpp
@@ -17,9 +17,9 @@ UInstancedActorsVisualizationProcessor::UInstancedActorsVisualizationProcessor()
 	UpdateParams.bTestCollisionAvailibilityForActorVisualization = false;
 }
 
-void UInstancedActorsVisualizationProcessor::ConfigureQueries() 
+void UInstancedActorsVisualizationProcessor::ConfigureQueries(const TSharedRef<FMassEntityManager>& EntityManager) 
 {
-	Super::ConfigureQueries();
+	Super::ConfigureQueries(EntityManager);
 
 	EntityQuery.ClearTagRequirements(FMassTagBitSet(*FMassVisualizationProcessorTag::StaticStruct()));
 	EntityQuery.AddTagRequirement<FInstancedActorsVisualizationProcessorTag>(EMassFragmentPresence::All);
diff --git a/Engine/Plugins/Runtime/InstancedActors/Source/InstancedActors/Private/InstancedActorsVisualizationSwitcherProcessor.cpp b/Engine/Plugins/Runtime/InstancedActors/Source/InstancedActors/Private/InstancedActorsVisualizationSwitcherProcessor.cpp
--- a/Engine/Plugins/Runtime/InstancedActors/Source/InstancedActors/Private/InstancedActorsVisualizationSwitcherProcessor.cpp
+++ b/Engine/Plugins/Runtime/InstancedActors/Source/InstancedActors/Private/InstancedActorsVisualizationSwitcherProcessor.cpp
@@ -22,7 +22,7 @@ UInstancedActorsVisualizationSwitcherProcessor::UInstancedActorsVisualizationSwi
 	ExecutionOrder.ExecuteBefore.Add(UMassStationaryISMSwitcherProcessor::StaticClass()->GetFName());
 }
 
-void UInstancedActorsVisualizationSwitcherProcessor::ConfigureQueries()
+void UInstancedActorsVisualizationSwitcherProcessor::ConfigureQueries(const TSharedRef<FMassEntityManager>& EntityManager)
 {
 	EntityQuery.AddRequirement<FInstancedActorsMeshSwitchFragment>(EMassFragmentAccess::ReadOnly);
 	EntityQuery.AddRequirement<FMassRepresentationFragment>(EMassFragmentAccess::ReadWrite);
diff --git a/Engine/Plugins/Runtime/InstancedActors/Source/InstancedActors/Public/InstancedActorsDebugProcessor.h b/Engine/Plugins/Runtime/InstancedActors/Source/InstancedActors/Public/InstancedActorsDebugProcessor.h
--- a/Engine/Plugins/Runtime/InstancedActors/Source/InstancedActors/Public/InstancedActorsDebugProcessor.h
+++ b/Engine/Plugins/Runtime/InstancedActors/Source/InstancedActors/Public/InstancedActorsDebugProcessor.h
@@ -18,7 +18,7 @@ class UInstancedActorsDebugProcessor : public UMassProcessor
 protected:
 	UInstancedActorsDebugProcessor();
 
-	virtual void ConfigureQueries() override;
+	virtual void ConfigureQueries(const TSharedRef<FMassEntityManager>& EntityManager) override;
 
 	virtual void Execute(FMassEntityManager& EntityManager, FMassExecutionContext& Context) override;
 
diff --git a/Engine/Plugins/Runtime/InstancedActors/Source/InstancedActors/Public/InstancedActorsInitializerProcessor.h b/Engine/Plugins/Runtime/InstancedActors/Source/InstancedActors/Public/InstancedActorsInitializerProcessor.h
--- a/Engine/Plugins/Runtime/InstancedActors/Source/InstancedActors/Public/InstancedActorsInitializerProcessor.h
+++ b/Engine/Plugins/Runtime/InstancedActors/Source/InstancedActors/Public/InstancedActorsInitializerProcessor.h
@@ -26,7 +26,7 @@ public:
 	UInstancedActorsInitializerProcessor();
 
 protected:
-	virtual void ConfigureQueries() override;
+	virtual void ConfigureQueries(const TSharedRef<FMassEntityManager>& EntityManager) override;
 	virtual void Execute(FMassEntityManager& EntityManager, FMassExecutionContext& Context) override;
 
 	FMassEntityQuery EntityQuery;
diff --git a/Engine/Plugins/Runtime/InstancedActors/Source/InstancedActors/Public/InstancedActorsStationaryLODBatchProcessor.h b/Engine/Plugins/Runtime/InstancedActors/Source/InstancedActors/Public/InstancedActorsStationaryLODBatchProcessor.h
--- a/Engine/Plugins/Runtime/InstancedActors/Source/InstancedActors/Public/InstancedActorsStationaryLODBatchProcessor.h
+++ b/Engine/Plugins/Runtime/InstancedActors/Source/InstancedActors/Public/InstancedActorsStationaryLODBatchProcessor.h
@@ -18,7 +18,7 @@ public:
 
 protected:
 	virtual bool ShouldAllowQueryBasedPruning(const bool bRuntimeMode = true) const override { return false; }
-	virtual void ConfigureQueries() override;
+	virtual void ConfigureQueries(const TSharedRef<FMassEntityManager>& EntityManager) override;
 	virtual void Execute(FMassEntityManager& EntityManager, FMassExecutionContext& Context) override;
 
 	FMassEntityQuery LODChangingEntityQuery;
diff --git a/Engine/Plugins/Runtime/InstancedActors/Source/InstancedActors/Public/InstancedActorsVisualizationProcessor.h b/Engine/Plugins/Runtime/InstancedActors/Source/InstancedActors/Public/InstancedActorsVisualizationProcessor.h
--- a/Engine/Plugins/Runtime/InstancedActors/Source/InstancedActors/Public/InstancedActorsVisualizationProcessor.h
+++ b/Engine/Plugins/Runtime/InstancedActors/Source/InstancedActors/Public/InstancedActorsVisualizationProcessor.h
@@ -23,5 +23,5 @@ class UInstancedActorsVisualizationProcessor : public UMassVisualizationProcesso
 
 protected:
 	UInstancedActorsVisualizationProcessor();
-	virtual void ConfigureQueries() override;
+	virtual void ConfigureQueries(const TSharedRef<FMassEntityManager>& EntityManager) override;
 };
diff --git a/Engine/Plugins/Runtime/InstancedActors/Source/InstancedActors/Public/InstancedActorsVisualizationSwitcherProcessor.h b/Engine/Plugins/Runtime/InstancedActors/Source/InstancedActors/Public/InstancedActorsVisualizationSwitcherProcessor.h
--- a/Engine/Plugins/Runtime/InstancedActors/Source/InstancedActors/Public/InstancedActorsVisualizationSwitcherProcessor.h
+++ b/Engine/Plugins/Runtime/InstancedActors/Source/InstancedActors/Public/InstancedActorsVisualizationSwitcherProcessor.h
@@ -27,7 +27,7 @@ public:
 	INSTANCEDACTORS_API static void SwitchEntityMeshDesc(FMassInstancedStaticMeshInfoArrayView& ISMInfosView, FMassRepresentationFragment& RepresentationFragment, FMassEntityHandle EntityHandle, FStaticMeshInstanceVisualizationDescHandle NewStaticMeshDescHandle);
 
 protected:
-	virtual void ConfigureQueries() override;
+	virtual void ConfigureQueries(const TSharedRef<FMassEntityManager>& EntityManager) override;
 	virtual void Execute(FMassEntityManager& EntityManager, FMassExecutionContext& Context) override;
 
 	FMassEntityQuery EntityQuery;
diff --git a/Engine/Plugins/Runtime/MassGameplay/Source/MassActors/Private/Translators/MassCapsuleComponentTranslators.cpp b/Engine/Plugins/Runtime/MassGameplay/Source/MassActors/Private/Translators/MassCapsuleComponentTranslators.cpp
--- a/Engine/Plugins/Runtime/MassGameplay/Source/MassActors/Private/Translators/MassCapsuleComponentTranslators.cpp
+++ b/Engine/Plugins/Runtime/MassGameplay/Source/MassActors/Private/Translators/MassCapsuleComponentTranslators.cpp
@@ -18,7 +18,7 @@ UMassCapsuleTransformToMassTranslator::UMassCapsuleTransformToMassTranslator()
 	ExecutionOrder.ExecuteInGroup = UE::Mass::ProcessorGroupNames::SyncWorldToMass;
 }
 
-void UMassCapsuleTransformToMassTranslator::ConfigureQueries()
+void UMassCapsuleTransformToMassTranslator::ConfigureQueries(const TSharedRef<FMassEntityManager>& EntityManager)
 {
 	AddRequiredTagsToQuery(EntityQuery);
 	EntityQuery.AddRequirement<FCapsuleComponentWrapperFragment>(EMassFragmentAccess::ReadOnly);
@@ -54,7 +54,7 @@ UMassTransformToActorCapsuleTranslator::UMassTransformToActorCapsuleTranslator()
 	bRequiresGameThreadExecution = true;
 }
 
-void UMassTransformToActorCapsuleTranslator::ConfigureQueries()
+void UMassTransformToActorCapsuleTranslator::ConfigureQueries(const TSharedRef<FMassEntityManager>& EntityManager)
 {
 	AddRequiredTagsToQuery(EntityQuery);
 	EntityQuery.AddRequirement<FCapsuleComponentWrapperFragment>(EMassFragmentAccess::ReadWrite);
diff --git a/Engine/Plugins/Runtime/MassGameplay/Source/MassActors/Private/Translators/MassCharacterMovementTranslators.cpp b/Engine/Plugins/Runtime/MassGameplay/Source/MassActors/Private/Translators/MassCharacterMovementTranslators.cpp
--- a/Engine/Plugins/Runtime/MassGameplay/Source/MassActors/Private/Translators/MassCharacterMovementTranslators.cpp
+++ b/Engine/Plugins/Runtime/MassGameplay/Source/MassActors/Private/Translators/MassCharacterMovementTranslators.cpp
@@ -22,7 +22,7 @@ UMassCharacterMovementToMassTranslator::UMassCharacterMovementToMassTranslator()
 	bRequiresGameThreadExecution = true;
 }
 
-void UMassCharacterMovementToMassTranslator::ConfigureQueries()
+void UMassCharacterMovementToMassTranslator::ConfigureQueries(const TSharedRef<FMassEntityManager>& EntityManager)
 {
 	AddRequiredTagsToQuery(EntityQuery);
 	EntityQuery.AddRequirement<FCharacterMovementComponentWrapperFragment>(EMassFragmentAccess::ReadOnly);
@@ -65,7 +65,7 @@ UMassCharacterMovementToActorTranslator::UMassCharacterMovementToActorTranslator
 	bRequiresGameThreadExecution = true;
 }
 
-void UMassCharacterMovementToActorTranslator::ConfigureQueries()
+void UMassCharacterMovementToActorTranslator::ConfigureQueries(const TSharedRef<FMassEntityManager>& EntityManager)
 {
 	AddRequiredTagsToQuery(EntityQuery);
 	EntityQuery.AddRequirement<FCharacterMovementComponentWrapperFragment>(EMassFragmentAccess::ReadWrite);
@@ -103,7 +103,7 @@ UMassCharacterOrientationToMassTranslator::UMassCharacterOrientationToMassTransl
 	bRequiresGameThreadExecution = true;
 }
 
-void UMassCharacterOrientationToMassTranslator::ConfigureQueries()
+void UMassCharacterOrientationToMassTranslator::ConfigureQueries(const TSharedRef<FMassEntityManager>& EntityManager)
 {
 	AddRequiredTagsToQuery(EntityQuery);
 	EntityQuery.AddRequirement<FCharacterMovementComponentWrapperFragment>(EMassFragmentAccess::ReadOnly);
@@ -145,7 +145,7 @@ UMassCharacterOrientationToActorTranslator::UMassCharacterOrientationToActorTran
 	bRequiresGameThreadExecution = true;
 }
 
-void UMassCharacterOrientationToActorTranslator::ConfigureQueries()
+void UMassCharacterOrientationToActorTranslator::ConfigureQueries(const TSharedRef<FMassEntityManager>& EntityManager)
 {
 	AddRequiredTagsToQuery(EntityQuery);
 	EntityQuery.AddRequirement<FCharacterMovementComponentWrapperFragment>(EMassFragmentAccess::ReadWrite);
diff --git a/Engine/Plugins/Runtime/MassGameplay/Source/MassActors/Private/Translators/MassSceneComponentLocationTranslator.cpp b/Engine/Plugins/Runtime/MassGameplay/Source/MassActors/Private/Translators/MassSceneComponentLocationTranslator.cpp
--- a/Engine/Plugins/Runtime/MassGameplay/Source/MassActors/Private/Translators/MassSceneComponentLocationTranslator.cpp
+++ b/Engine/Plugins/Runtime/MassGameplay/Source/MassActors/Private/Translators/MassSceneComponentLocationTranslator.cpp
@@ -16,7 +16,7 @@ UMassSceneComponentLocationToMassTranslator::UMassSceneComponentLocationToMassTr
 	RequiredTags.Add<FMassSceneComponentLocationCopyToMassTag>();
 }
 
-void UMassSceneComponentLocationToMassTranslator::ConfigureQueries()
+void UMassSceneComponentLocationToMassTranslator::ConfigureQueries(const TSharedRef<FMassEntityManager>& EntityManager)
 {
 	AddRequiredTagsToQuery(EntityQuery);
 	EntityQuery.AddRequirement<FMassSceneComponentWrapperFragment>(EMassFragmentAccess::ReadOnly);
@@ -54,7 +54,7 @@ UMassSceneComponentLocationToActorTranslator::UMassSceneComponentLocationToActor
 	bRequiresGameThreadExecution = true;
 }
 
-void UMassSceneComponentLocationToActorTranslator::ConfigureQueries()
+void UMassSceneComponentLocationToActorTranslator::ConfigureQueries(const TSharedRef<FMassEntityManager>& EntityManager)
 {
 	AddRequiredTagsToQuery(EntityQuery);
 	EntityQuery.AddRequirement<FMassSceneComponentWrapperFragment>(EMassFragmentAccess::ReadOnly);
diff --git a/Engine/Plugins/Runtime/MassGameplay/Source/MassActors/Private/Translators/MassTranslators_BehaviorTree.cpp b/Engine/Plugins/Runtime/MassGameplay/Source/MassActors/Private/Translators/MassTranslators_BehaviorTree.cpp
--- a/Engine/Plugins/Runtime/MassGameplay/Source/MassActors/Private/Translators/MassTranslators_BehaviorTree.cpp
+++ b/Engine/Plugins/Runtime/MassGameplay/Source/MassActors/Private/Translators/MassTranslators_BehaviorTree.cpp
@@ -10,7 +10,7 @@ UMassTranslator_BehaviorTree::UMassTranslator_BehaviorTree()
 	bAutoRegisterWithProcessingPhases = false;
 }
 
-void UMassTranslator_BehaviorTree::ConfigureQueries()
+void UMassTranslator_BehaviorTree::ConfigureQueries(const TSharedRef<FMassEntityManager>& EntityManager)
 {
 	EntityQuery.AddRequirement<FDataFragment_BehaviorTreeComponentWrapper>(EMassFragmentAccess::ReadWrite);
 }
diff --git a/Engine/Plugins/Runtime/MassGameplay/Source/MassActors/Public/Translators/MassCapsuleComponentTranslators.h b/Engine/Plugins/Runtime/MassGameplay/Source/MassActors/Public/Translators/MassCapsuleComponentTranslators.h
--- a/Engine/Plugins/Runtime/MassGameplay/Source/MassActors/Public/Translators/MassCapsuleComponentTranslators.h
+++ b/Engine/Plugins/Runtime/MassGameplay/Source/MassActors/Public/Translators/MassCapsuleComponentTranslators.h
@@ -39,7 +39,7 @@ public:
 	UMassCapsuleTransformToMassTranslator();
 
 protected:
-	virtual void ConfigureQueries() override;
+	virtual void ConfigureQueries(const TSharedRef<FMassEntityManager>& EntityManager) override;
 	virtual void Execute(FMassEntityManager& EntityManager, FMassExecutionContext& Context) override;	
 
 	FMassEntityQuery EntityQuery;
@@ -59,7 +59,7 @@ public:
 	UMassTransformToActorCapsuleTranslator();
 
 protected:
-	virtual void ConfigureQueries() override;
+	virtual void ConfigureQueries(const TSharedRef<FMassEntityManager>& EntityManager) override;
 	virtual void Execute(FMassEntityManager& EntityManager, FMassExecutionContext& Context) override;
 
 	FMassEntityQuery EntityQuery;
diff --git a/Engine/Plugins/Runtime/MassGameplay/Source/MassActors/Public/Translators/MassCharacterMovementTranslators.h b/Engine/Plugins/Runtime/MassGameplay/Source/MassActors/Public/Translators/MassCharacterMovementTranslators.h
--- a/Engine/Plugins/Runtime/MassGameplay/Source/MassActors/Public/Translators/MassCharacterMovementTranslators.h
+++ b/Engine/Plugins/Runtime/MassGameplay/Source/MassActors/Public/Translators/MassCharacterMovementTranslators.h
@@ -31,7 +31,7 @@ public:
 	UMassCharacterMovementToMassTranslator();
 
 protected:
-	virtual void ConfigureQueries() override;
+	virtual void ConfigureQueries(const TSharedRef<FMassEntityManager>& EntityManager) override;
 	virtual void Execute(FMassEntityManager& EntityManager, FMassExecutionContext& Context) override;
 
 	FMassEntityQuery EntityQuery;
@@ -53,7 +53,7 @@ public:
 	UMassCharacterMovementToActorTranslator();
 
 protected:
-	virtual void ConfigureQueries() override;
+	virtual void ConfigureQueries(const TSharedRef<FMassEntityManager>& EntityManager) override;
 	virtual void Execute(FMassEntityManager& EntityManager, FMassExecutionContext& Context) override;
 
 	FMassEntityQuery EntityQuery;
@@ -75,7 +75,7 @@ public:
 	UMassCharacterOrientationToMassTranslator();
 
 protected:
-	virtual void ConfigureQueries() override;
+	virtual void ConfigureQueries(const TSharedRef<FMassEntityManager>& EntityManager) override;
 	virtual void Execute(FMassEntityManager& EntityManager, FMassExecutionContext& Context) override;
 
 	FMassEntityQuery EntityQuery;
@@ -97,7 +97,7 @@ public:
 	UMassCharacterOrientationToActorTranslator();
 
 protected:
-	virtual void ConfigureQueries() override;
+	virtual void ConfigureQueries(const TSharedRef<FMassEntityManager>& EntityManager) override;
 	virtual void Execute(FMassEntityManager& EntityManager, FMassExecutionContext& Context) override;
 
 	FMassEntityQuery EntityQuery;
diff --git a/Engine/Plugins/Runtime/MassGameplay/Source/MassActors/Public/Translators/MassSceneComponentLocationTranslator.h b/Engine/Plugins/Runtime/MassGameplay/Source/MassActors/Public/Translators/MassSceneComponentLocationTranslator.h
--- a/Engine/Plugins/Runtime/MassGameplay/Source/MassActors/Public/Translators/MassSceneComponentLocationTranslator.h
+++ b/Engine/Plugins/Runtime/MassGameplay/Source/MassActors/Public/Translators/MassSceneComponentLocationTranslator.h
@@ -29,7 +29,7 @@ public:
 	UMassSceneComponentLocationToMassTranslator();
 
 protected:
-	virtual void ConfigureQueries() override;
+	virtual void ConfigureQueries(const TSharedRef<FMassEntityManager>& EntityManager) override;
 	virtual void Execute(FMassEntityManager& EntityManager, FMassExecutionContext& Context) override;
 
 	FMassEntityQuery EntityQuery;
@@ -51,7 +51,7 @@ public:
 	UMassSceneComponentLocationToActorTranslator();
 
 protected:
-	virtual void ConfigureQueries() override;
+	virtual void ConfigureQueries(const TSharedRef<FMassEntityManager>& EntityManager) override;
 	virtual void Execute(FMassEntityManager& EntityManager, FMassExecutionContext& Context) override;
 
 	FMassEntityQuery EntityQuery;
diff --git a/Engine/Plugins/Runtime/MassGameplay/Source/MassActors/Public/Translators/MassTranslators_BehaviorTree.h b/Engine/Plugins/Runtime/MassGameplay/Source/MassActors/Public/Translators/MassTranslators_BehaviorTree.h
--- a/Engine/Plugins/Runtime/MassGameplay/Source/MassActors/Public/Translators/MassTranslators_BehaviorTree.h
+++ b/Engine/Plugins/Runtime/MassGameplay/Source/MassActors/Public/Translators/MassTranslators_BehaviorTree.h
@@ -25,7 +25,7 @@ public:
 	UMassTranslator_BehaviorTree();
 
 protected:
-	virtual void ConfigureQueries() override;
+	virtual void ConfigureQueries(const TSharedRef<FMassEntityManager>& EntityManager) override;
 	virtual void Execute(FMassEntityManager& EntityManager, FMassExecutionContext& Context) override {}
 
 	FMassEntityQuery EntityQuery;
diff --git a/Engine/Plugins/Runtime/MassGameplay/Source/MassEQS/Private/Processors/MassEnvQueryGeneratorProcessor_MassEntityHandles.cpp b/Engine/Plugins/Runtime/MassGameplay/Source/MassEQS/Private/Processors/MassEnvQueryGeneratorProcessor_MassEntityHandles.cpp
--- a/Engine/Plugins/Runtime/MassGameplay/Source/MassEQS/Private/Processors/MassEnvQueryGeneratorProcessor_MassEntityHandles.cpp
+++ b/Engine/Plugins/Runtime/MassGameplay/Source/MassEQS/Private/Processors/MassEnvQueryGeneratorProcessor_MassEntityHandles.cpp
@@ -15,7 +15,7 @@ UMassEnvQueryGeneratorProcessor_MassEntityHandles::UMassEnvQueryGeneratorProcess
 	CorrespondingRequestClass = UMassEnvQueryGenerator_MassEntityHandles::StaticClass();
 }
 
-void UMassEnvQueryGeneratorProcessor_MassEntityHandles::ConfigureQueries()
+void UMassEnvQueryGeneratorProcessor_MassEntityHandles::ConfigureQueries(const TSharedRef<FMassEntityManager>& EntityManager)
 {
 	EntityQuery.AddRequirement<FTransformFragment>(EMassFragmentAccess::ReadOnly);
 	ProcessorRequirements.AddSubsystemRequirement<UMassEQSSubsystem>(EMassFragmentAccess::ReadWrite);
diff --git a/Engine/Plugins/Runtime/MassGameplay/Source/MassEQS/Private/Processors/MassEnvQueryTestProcessor_MassEntityTags.cpp b/Engine/Plugins/Runtime/MassGameplay/Source/MassEQS/Private/Processors/MassEnvQueryTestProcessor_MassEntityTags.cpp
--- a/Engine/Plugins/Runtime/MassGameplay/Source/MassEQS/Private/Processors/MassEnvQueryTestProcessor_MassEntityTags.cpp
+++ b/Engine/Plugins/Runtime/MassGameplay/Source/MassEQS/Private/Processors/MassEnvQueryTestProcessor_MassEntityTags.cpp
@@ -24,7 +24,7 @@ UMassEnvQueryTestProcessor_MassEntityTags::UMassEnvQueryTestProcessor_MassEntity
 	CorrespondingRequestClass = UMassEnvQueryTest_MassEntityTags::StaticClass();
 }
 
-void UMassEnvQueryTestProcessor_MassEntityTags::ConfigureQueries()
+void UMassEnvQueryTestProcessor_MassEntityTags::ConfigureQueries(const TSharedRef<FMassEntityManager>& EntityManager)
 {
 	ProcessorRequirements.AddSubsystemRequirement<UMassEQSSubsystem>(EMassFragmentAccess::ReadWrite);
 }
diff --git a/Engine/Plugins/Runtime/MassGameplay/Source/MassEQS/Public/Processors/MassEnvQueryGeneratorProcessor_MassEntityHandles.h b/Engine/Plugins/Runtime/MassGameplay/Source/MassEQS/Public/Processors/MassEnvQueryGeneratorProcessor_MassEntityHandles.h
--- a/Engine/Plugins/Runtime/MassGameplay/Source/MassEQS/Public/Processors/MassEnvQueryGeneratorProcessor_MassEntityHandles.h
+++ b/Engine/Plugins/Runtime/MassGameplay/Source/MassEQS/Public/Processors/MassEnvQueryGeneratorProcessor_MassEntityHandles.h
@@ -14,7 +14,7 @@ public:
 	UMassEnvQueryGeneratorProcessor_MassEntityHandles();
 
 protected:
-	virtual void ConfigureQueries() override;
+	virtual void ConfigureQueries(const TSharedRef<FMassEntityManager>& EntityManager) override;
 	virtual void Execute(FMassEntityManager& EntityManager, FMassExecutionContext& Context) override;
 
 	FMassEntityQuery EntityQuery;
diff --git a/Engine/Plugins/Runtime/MassGameplay/Source/MassEQS/Public/Processors/MassEnvQueryTestProcessor_MassEntityTags.h b/Engine/Plugins/Runtime/MassGameplay/Source/MassEQS/Public/Processors/MassEnvQueryTestProcessor_MassEntityTags.h
--- a/Engine/Plugins/Runtime/MassGameplay/Source/MassEQS/Public/Processors/MassEnvQueryTestProcessor_MassEntityTags.h
+++ b/Engine/Plugins/Runtime/MassGameplay/Source/MassEQS/Public/Processors/MassEnvQueryTestProcessor_MassEntityTags.h
@@ -14,7 +14,7 @@ public:
 	UMassEnvQueryTestProcessor_MassEntityTags();
 
 protected:
-	virtual void ConfigureQueries() override;
+	virtual void ConfigureQueries(const TSharedRef<FMassEntityManager>& EntityManager) override;
 	virtual void Execute(FMassEntityManager& EntityManager, FMassExecutionContext& Context) override;
 	virtual bool ShouldAllowQueryBasedPruning(const bool bRuntimeMode = true) const override { return false; }
 
diff --git a/Engine/Plugins/Runtime/MassGameplay/Source/MassGameplayDebug/Private/AssignDebugVisProcessor.cpp b/Engine/Plugins/Runtime/MassGameplay/Source/MassGameplayDebug/Private/AssignDebugVisProcessor.cpp
--- a/Engine/Plugins/Runtime/MassGameplay/Source/MassGameplayDebug/Private/AssignDebugVisProcessor.cpp
+++ b/Engine/Plugins/Runtime/MassGameplay/Source/MassGameplayDebug/Private/AssignDebugVisProcessor.cpp
@@ -18,7 +18,7 @@ UAssignDebugVisProcessor::UAssignDebugVisProcessor()
 	Operation = EMassObservedOperation::Add;
 }
 
-void UAssignDebugVisProcessor::ConfigureQueries()
+void UAssignDebugVisProcessor::ConfigureQueries(const TSharedRef<FMassEntityManager>& EntityManager)
 {
 	EntityQuery.AddRequirement<FSimDebugVisFragment>(EMassFragmentAccess::ReadWrite);
 	EntityQuery.AddSubsystemRequirement<UMassDebuggerSubsystem>(EMassFragmentAccess::ReadWrite);
diff --git a/Engine/Plugins/Runtime/MassGameplay/Source/MassGameplayDebug/Private/DebugVisLocationProcessor.cpp b/Engine/Plugins/Runtime/MassGameplay/Source/MassGameplayDebug/Private/DebugVisLocationProcessor.cpp
--- a/Engine/Plugins/Runtime/MassGameplay/Source/MassGameplayDebug/Private/DebugVisLocationProcessor.cpp
+++ b/Engine/Plugins/Runtime/MassGameplay/Source/MassGameplayDebug/Private/DebugVisLocationProcessor.cpp
@@ -19,7 +19,7 @@ UDebugVisLocationProcessor::UDebugVisLocationProcessor()
 	bRequiresGameThreadExecution = true; // due to UMassDebuggerSubsystem access
 }
 
-void UDebugVisLocationProcessor::ConfigureQueries()
+void UDebugVisLocationProcessor::ConfigureQueries(const TSharedRef<FMassEntityManager>& EntityManager)
 {
 	EntityQuery.AddRequirement<FTransformFragment>(EMassFragmentAccess::ReadOnly);
 	EntityQuery.AddRequirement<FSimDebugVisFragment>(EMassFragmentAccess::ReadOnly);
@@ -80,7 +80,7 @@ UMassProcessor_UpdateDebugVis::UMassProcessor_UpdateDebugVis()
 	bRequiresGameThreadExecution = true; // due to UMassDebuggerSubsystem access
 }
 
-void UMassProcessor_UpdateDebugVis::ConfigureQueries() 
+void UMassProcessor_UpdateDebugVis::ConfigureQueries(const TSharedRef<FMassEntityManager>& EntityManager) 
 {
 	EntityQuery.AddRequirement<FTransformFragment>(EMassFragmentAccess::ReadOnly);
 	EntityQuery.AddRequirement<FDataFragment_DebugVis>(EMassFragmentAccess::ReadWrite);
diff --git a/Engine/Plugins/Runtime/MassGameplay/Source/MassGameplayDebug/Public/AssignDebugVisProcessor.h b/Engine/Plugins/Runtime/MassGameplay/Source/MassGameplayDebug/Public/AssignDebugVisProcessor.h
--- a/Engine/Plugins/Runtime/MassGameplay/Source/MassGameplayDebug/Public/AssignDebugVisProcessor.h
+++ b/Engine/Plugins/Runtime/MassGameplay/Source/MassGameplayDebug/Public/AssignDebugVisProcessor.h
@@ -15,7 +15,7 @@ class MASSGAMEPLAYDEBUG_API UAssignDebugVisProcessor : public UMassObserverProce
 	GENERATED_BODY()
 public:
 	UAssignDebugVisProcessor();
-	virtual void ConfigureQueries() override;
+	virtual void ConfigureQueries(const TSharedRef<FMassEntityManager>& EntityManager) override;
 	virtual void Execute(FMassEntityManager& EntityManager, FMassExecutionContext& Context) override;
 
 protected:
diff --git a/Engine/Plugins/Runtime/MassGameplay/Source/MassGameplayDebug/Public/DebugVisLocationProcessor.h b/Engine/Plugins/Runtime/MassGameplay/Source/MassGameplayDebug/Public/DebugVisLocationProcessor.h
--- a/Engine/Plugins/Runtime/MassGameplay/Source/MassGameplayDebug/Public/DebugVisLocationProcessor.h
+++ b/Engine/Plugins/Runtime/MassGameplay/Source/MassGameplayDebug/Public/DebugVisLocationProcessor.h
@@ -17,7 +17,7 @@ public:
 	UDebugVisLocationProcessor();
 	
 protected:
-	virtual void ConfigureQueries() override;
+	virtual void ConfigureQueries(const TSharedRef<FMassEntityManager>& EntityManager) override;
 	virtual void Execute(FMassEntityManager& EntityManager, FMassExecutionContext& Context) override;
 
 	FMassEntityQuery EntityQuery;
@@ -34,7 +34,7 @@ class MASSGAMEPLAYDEBUG_API UMassProcessor_UpdateDebugVis : public UMassProcesso
 	GENERATED_BODY()
 public:
 	UMassProcessor_UpdateDebugVis();
-	virtual void ConfigureQueries() override;
+	virtual void ConfigureQueries(const TSharedRef<FMassEntityManager>& EntityManager) override;
 	virtual void Execute(FMassEntityManager& EntityManager, FMassExecutionContext& Context) override;
 
 protected:
diff --git a/Engine/Plugins/Runtime/MassGameplay/Source/MassLOD/Private/MassLODCollectorProcessor.cpp b/Engine/Plugins/Runtime/MassGameplay/Source/MassLOD/Private/MassLODCollectorProcessor.cpp
--- a/Engine/Plugins/Runtime/MassGameplay/Source/MassLOD/Private/MassLODCollectorProcessor.cpp
+++ b/Engine/Plugins/Runtime/MassGameplay/Source/MassLOD/Private/MassLODCollectorProcessor.cpp
@@ -17,9 +17,9 @@ UMassLODCollectorProcessor::UMassLODCollectorProcessor()
 	ExecutionOrder.ExecuteAfter.Add(UE::Mass::ProcessorGroupNames::SyncWorldToMass);
 }
 
-void UMassLODCollectorProcessor::ConfigureQueries()
+void UMassLODCollectorProcessor::ConfigureQueries(const TSharedRef<FMassEntityManager>& EntityManager)
 {
-	FMassEntityQuery BaseQuery;
+	FMassEntityQuery BaseQuery(EntityManager);
 	BaseQuery.AddTagRequirement<FMassCollectLODViewerInfoTag>(EMassFragmentPresence::All);
 	BaseQuery.AddRequirement<FTransformFragment>(EMassFragmentAccess::ReadOnly);
 	BaseQuery.AddRequirement<FMassViewerInfoFragment>(EMassFragmentAccess::ReadWrite);
diff --git a/Engine/Plugins/Runtime/MassGameplay/Source/MassLOD/Private/MassLODDistanceCollectorProcessor.cpp b/Engine/Plugins/Runtime/MassGameplay/Source/MassLOD/Private/MassLODDistanceCollectorProcessor.cpp
--- a/Engine/Plugins/Runtime/MassGameplay/Source/MassLOD/Private/MassLODDistanceCollectorProcessor.cpp
+++ b/Engine/Plugins/Runtime/MassGameplay/Source/MassLOD/Private/MassLODDistanceCollectorProcessor.cpp
@@ -21,9 +21,9 @@ UMassLODDistanceCollectorProcessor::UMassLODDistanceCollectorProcessor()
 	ExecutionOrder.ExecuteAfter.Add(UE::Mass::ProcessorGroupNames::SyncWorldToMass);
 }
 
-void UMassLODDistanceCollectorProcessor::ConfigureQueries()
+void UMassLODDistanceCollectorProcessor::ConfigureQueries(const TSharedRef<FMassEntityManager>& EntityManager)
 {
-	FMassEntityQuery BaseQuery;
+	FMassEntityQuery BaseQuery(EntityManager);
 	BaseQuery.AddTagRequirement<FMassCollectDistanceLODViewerInfoTag>(EMassFragmentPresence::All);
 	BaseQuery.AddRequirement<FTransformFragment>(EMassFragmentAccess::ReadOnly);
 	BaseQuery.AddRequirement<FMassViewerInfoFragment>(EMassFragmentAccess::ReadWrite);
diff --git a/Engine/Plugins/Runtime/MassGameplay/Source/MassLOD/Private/MassSimulationLOD.cpp b/Engine/Plugins/Runtime/MassGameplay/Source/MassLOD/Private/MassSimulationLOD.cpp
--- a/Engine/Plugins/Runtime/MassGameplay/Source/MassLOD/Private/MassSimulationLOD.cpp
+++ b/Engine/Plugins/Runtime/MassGameplay/Source/MassLOD/Private/MassSimulationLOD.cpp
@@ -74,7 +74,7 @@ UMassSimulationLODProcessor::UMassSimulationLODProcessor()
 	ExecutionOrder.ExecuteAfter.Add(UE::Mass::ProcessorGroupNames::LODCollector);
 }
 
-void UMassSimulationLODProcessor::ConfigureQueries()
+void UMassSimulationLODProcessor::ConfigureQueries(const TSharedRef<FMassEntityManager>& EntityManager)
 {
 	EntityQuery.AddRequirement<FTransformFragment>(EMassFragmentAccess::ReadOnly);
 	EntityQuery.AddRequirement<FMassViewerInfoFragment>(EMassFragmentAccess::ReadOnly);
diff --git a/Engine/Plugins/Runtime/MassGameplay/Source/MassLOD/Public/MassLODCollectorProcessor.h b/Engine/Plugins/Runtime/MassGameplay/Source/MassLOD/Public/MassLODCollectorProcessor.h
--- a/Engine/Plugins/Runtime/MassGameplay/Source/MassLOD/Public/MassLODCollectorProcessor.h
+++ b/Engine/Plugins/Runtime/MassGameplay/Source/MassLOD/Public/MassLODCollectorProcessor.h
@@ -30,7 +30,7 @@ public:
 	UMassLODCollectorProcessor();
 
 protected:
-	virtual void ConfigureQueries() override;
+	virtual void ConfigureQueries(const TSharedRef<FMassEntityManager>& EntityManager) override;
 	virtual void Execute(FMassEntityManager& EntityManager, FMassExecutionContext& Context) override;
 
 	template <bool bLocalViewersOnly>
diff --git a/Engine/Plugins/Runtime/MassGameplay/Source/MassLOD/Public/MassLODDistanceCollectorProcessor.h b/Engine/Plugins/Runtime/MassGameplay/Source/MassLOD/Public/MassLODDistanceCollectorProcessor.h
--- a/Engine/Plugins/Runtime/MassGameplay/Source/MassLOD/Public/MassLODDistanceCollectorProcessor.h
+++ b/Engine/Plugins/Runtime/MassGameplay/Source/MassLOD/Public/MassLODDistanceCollectorProcessor.h
@@ -24,7 +24,7 @@ public:
 	UMassLODDistanceCollectorProcessor();
 
 protected:
-	virtual void ConfigureQueries() override;
+	virtual void ConfigureQueries(const TSharedRef<FMassEntityManager>& EntityManager) override;
 	virtual void Execute(FMassEntityManager& EntityManager, FMassExecutionContext& Context) override;
 
 	template <bool bLocalViewersOnly>
diff --git a/Engine/Plugins/Runtime/MassGameplay/Source/MassLOD/Public/MassSimulationLOD.h b/Engine/Plugins/Runtime/MassGameplay/Source/MassLOD/Public/MassSimulationLOD.h
--- a/Engine/Plugins/Runtime/MassGameplay/Source/MassLOD/Public/MassSimulationLOD.h
+++ b/Engine/Plugins/Runtime/MassGameplay/Source/MassLOD/Public/MassSimulationLOD.h
@@ -170,7 +170,7 @@ public:
 
 protected:
 
-	virtual void ConfigureQueries() override;
+	virtual void ConfigureQueries(const TSharedRef<FMassEntityManager>& EntityManager) override;
 	virtual void Execute(FMassEntityManager& EntityManager, FMassExecutionContext& Context) override;
 
 	void CalculateLODForConfig(FMassEntityManager& EntityManager, FMassExecutionContext& Context, FMassSimulationLODParameters& LODParams);
diff --git a/Engine/Plugins/Runtime/MassGameplay/Source/MassMovement/Private/MassMovementProcessors.cpp b/Engine/Plugins/Runtime/MassGameplay/Source/MassMovement/Private/MassMovementProcessors.cpp
--- a/Engine/Plugins/Runtime/MassGameplay/Source/MassMovement/Private/MassMovementProcessors.cpp
+++ b/Engine/Plugins/Runtime/MassGameplay/Source/MassMovement/Private/MassMovementProcessors.cpp
@@ -24,7 +24,7 @@ UMassApplyMovementProcessor::UMassApplyMovementProcessor()
 	ExecutionOrder.ExecuteAfter.Add(UE::Mass::ProcessorGroupNames::Avoidance);
 }
 
-void UMassApplyMovementProcessor::ConfigureQueries()
+void UMassApplyMovementProcessor::ConfigureQueries(const TSharedRef<FMassEntityManager>& EntityManager)
 {
 	EntityQuery.AddRequirement<FMassVelocityFragment>(EMassFragmentAccess::ReadWrite);
 	EntityQuery.AddRequirement<FTransformFragment>(EMassFragmentAccess::ReadWrite);
diff --git a/Engine/Plugins/Runtime/MassGameplay/Source/MassMovement/Private/MassSimpleMovementTrait.cpp b/Engine/Plugins/Runtime/MassGameplay/Source/MassMovement/Private/MassSimpleMovementTrait.cpp
--- a/Engine/Plugins/Runtime/MassGameplay/Source/MassMovement/Private/MassSimpleMovementTrait.cpp
+++ b/Engine/Plugins/Runtime/MassGameplay/Source/MassMovement/Private/MassSimpleMovementTrait.cpp
@@ -30,7 +30,7 @@ UMassSimpleMovementProcessor::UMassSimpleMovementProcessor()
 	ExecutionOrder.ExecuteInGroup = UE::Mass::ProcessorGroupNames::Avoidance;
 }
 
-void UMassSimpleMovementProcessor::ConfigureQueries()
+void UMassSimpleMovementProcessor::ConfigureQueries(const TSharedRef<FMassEntityManager>& EntityManager)
 {
 	EntityQuery.AddRequirement<FMassVelocityFragment>(EMassFragmentAccess::ReadOnly);
 	EntityQuery.AddRequirement<FTransformFragment>(EMassFragmentAccess::ReadWrite);
diff --git a/Engine/Plugins/Runtime/MassGameplay/Source/MassMovement/Private/MassVelocityRandomizerTrait.cpp b/Engine/Plugins/Runtime/MassGameplay/Source/MassMovement/Private/MassVelocityRandomizerTrait.cpp
--- a/Engine/Plugins/Runtime/MassGameplay/Source/MassMovement/Private/MassVelocityRandomizerTrait.cpp
+++ b/Engine/Plugins/Runtime/MassGameplay/Source/MassMovement/Private/MassVelocityRandomizerTrait.cpp
@@ -31,7 +31,7 @@ UMassRandomVelocityInitializer::UMassRandomVelocityInitializer()
 	RandomStream.Initialize(UE::Mass::Utils::GenerateRandomSeed());
 }
 
-void UMassRandomVelocityInitializer::ConfigureQueries()
+void UMassRandomVelocityInitializer::ConfigureQueries(const TSharedRef<FMassEntityManager>& EntityManager)
 {
 	EntityQuery.AddRequirement<FMassVelocityFragment>(EMassFragmentAccess::ReadWrite);
 }
diff --git a/Engine/Plugins/Runtime/MassGameplay/Source/MassMovement/Public/Example/MassSimpleMovementTrait.h b/Engine/Plugins/Runtime/MassGameplay/Source/MassMovement/Public/Example/MassSimpleMovementTrait.h
--- a/Engine/Plugins/Runtime/MassGameplay/Source/MassMovement/Public/Example/MassSimpleMovementTrait.h
+++ b/Engine/Plugins/Runtime/MassGameplay/Source/MassMovement/Public/Example/MassSimpleMovementTrait.h
@@ -33,7 +33,7 @@ public:
 	UMassSimpleMovementProcessor();
 		
 protected:	
-	virtual void ConfigureQueries() override;
+	virtual void ConfigureQueries(const TSharedRef<FMassEntityManager>& EntityManager) override;
 	virtual void Execute(FMassEntityManager& EntityManager, FMassExecutionContext& Context) override;
 
 	FMassEntityQuery EntityQuery;
diff --git a/Engine/Plugins/Runtime/MassGameplay/Source/MassMovement/Public/Example/MassVelocityRandomizerTrait.h b/Engine/Plugins/Runtime/MassGameplay/Source/MassMovement/Public/Example/MassVelocityRandomizerTrait.h
--- a/Engine/Plugins/Runtime/MassGameplay/Source/MassMovement/Public/Example/MassVelocityRandomizerTrait.h
+++ b/Engine/Plugins/Runtime/MassGameplay/Source/MassMovement/Public/Example/MassVelocityRandomizerTrait.h
@@ -40,7 +40,7 @@ public:
 	void SetParameters(const float InMinSpeed, const float InMaxSpeed, const bool bInSetZComponent);
 
 protected:
-	virtual void ConfigureQueries() override;
+	virtual void ConfigureQueries(const TSharedRef<FMassEntityManager>& EntityManager) override;
 	virtual void Execute(FMassEntityManager& EntityManager, FMassExecutionContext& Context) override;
 
 	FMassEntityQuery EntityQuery;
diff --git a/Engine/Plugins/Runtime/MassGameplay/Source/MassMovement/Public/Movement/MassMovementProcessors.h b/Engine/Plugins/Runtime/MassGameplay/Source/MassMovement/Public/Movement/MassMovementProcessors.h
--- a/Engine/Plugins/Runtime/MassGameplay/Source/MassMovement/Public/Movement/MassMovementProcessors.h
+++ b/Engine/Plugins/Runtime/MassGameplay/Source/MassMovement/Public/Movement/MassMovementProcessors.h
@@ -18,7 +18,7 @@ public:
 	UMassApplyMovementProcessor();
 
 protected:
-	virtual void ConfigureQueries() override;
+	virtual void ConfigureQueries(const TSharedRef<FMassEntityManager>& EntityManager) override;
 	virtual void Execute(FMassEntityManager& EntityManager, FMassExecutionContext& Context) override;
 
 private:
diff --git a/Engine/Plugins/Runtime/MassGameplay/Source/MassReplication/Private/MassReplicationFragments.cpp b/Engine/Plugins/Runtime/MassGameplay/Source/MassReplication/Private/MassReplicationFragments.cpp
--- a/Engine/Plugins/Runtime/MassGameplay/Source/MassReplication/Private/MassReplicationFragments.cpp
+++ b/Engine/Plugins/Runtime/MassGameplay/Source/MassReplication/Private/MassReplicationFragments.cpp
@@ -17,7 +17,7 @@ UMassNetworkIDFragmentInitializer::UMassNetworkIDFragmentInitializer()
 	Operation = EMassObservedOperation::Add;
 }
 
-void UMassNetworkIDFragmentInitializer::ConfigureQueries()
+void UMassNetworkIDFragmentInitializer::ConfigureQueries(const TSharedRef<FMassEntityManager>& EntityManager)
 {
 	EntityQuery.AddRequirement<FMassNetworkIDFragment>(EMassFragmentAccess::ReadWrite);
 	EntityQuery.AddSubsystemRequirement<UMassReplicationSubsystem>(EMassFragmentAccess::ReadWrite);
diff --git a/Engine/Plugins/Runtime/MassGameplay/Source/MassReplication/Private/MassReplicationGridProcessor.cpp b/Engine/Plugins/Runtime/MassGameplay/Source/MassReplication/Private/MassReplicationGridProcessor.cpp
--- a/Engine/Plugins/Runtime/MassGameplay/Source/MassReplication/Private/MassReplicationGridProcessor.cpp
+++ b/Engine/Plugins/Runtime/MassGameplay/Source/MassReplication/Private/MassReplicationGridProcessor.cpp
@@ -25,7 +25,7 @@ UMassReplicationGridProcessor::UMassReplicationGridProcessor()
 	ProcessingPhase = EMassProcessingPhase::PostPhysics;
 }
 
-void UMassReplicationGridProcessor::ConfigureQueries()
+void UMassReplicationGridProcessor::ConfigureQueries(const TSharedRef<FMassEntityManager>& EntityManager)
 {
 	AddToGridEntityQuery.AddRequirement<FTransformFragment>(EMassFragmentAccess::ReadOnly);
 	AddToGridEntityQuery.AddRequirement<FMassReplicationGridCellLocationFragment>(EMassFragmentAccess::ReadWrite);
@@ -139,7 +139,7 @@ UMassReplicationGridRemoverProcessor::UMassReplicationGridRemoverProcessor()
 	ExecutionFlags = (int32)EProcessorExecutionFlags::AllNetModes;
 }
 
-void UMassReplicationGridRemoverProcessor::ConfigureQueries()
+void UMassReplicationGridRemoverProcessor::ConfigureQueries(const TSharedRef<FMassEntityManager>& EntityManager)
 {
 	EntityQuery.AddRequirement<FMassReplicationGridCellLocationFragment>(EMassFragmentAccess::ReadWrite);
 	EntityQuery.AddSubsystemRequirement<UMassReplicationSubsystem>(EMassFragmentAccess::ReadWrite);
diff --git a/Engine/Plugins/Runtime/MassGameplay/Source/MassReplication/Private/MassReplicationProcessor.cpp b/Engine/Plugins/Runtime/MassGameplay/Source/MassReplication/Private/MassReplicationProcessor.cpp
--- a/Engine/Plugins/Runtime/MassGameplay/Source/MassReplication/Private/MassReplicationProcessor.cpp
+++ b/Engine/Plugins/Runtime/MassGameplay/Source/MassReplication/Private/MassReplicationProcessor.cpp
@@ -35,7 +35,7 @@ UMassReplicationProcessor::UMassReplicationProcessor()
 	bRequiresGameThreadExecution = true;
 }
 
-void UMassReplicationProcessor::ConfigureQueries()
+void UMassReplicationProcessor::ConfigureQueries(const TSharedRef<FMassEntityManager>& EntityManager)
 {
 	SyncClientData.AddRequirement<FMassReplicationLODFragment>(EMassFragmentAccess::ReadWrite);
 	SyncClientData.AddRequirement<FMassReplicatedAgentFragment>(EMassFragmentAccess::ReadWrite);
diff --git a/Engine/Plugins/Runtime/MassGameplay/Source/MassReplication/Public/MassClientBubbleHandler.h b/Engine/Plugins/Runtime/MassGameplay/Source/MassReplication/Public/MassClientBubbleHandler.h
--- a/Engine/Plugins/Runtime/MassGameplay/Source/MassReplication/Public/MassClientBubbleHandler.h
+++ b/Engine/Plugins/Runtime/MassGameplay/Source/MassReplication/Public/MassClientBubbleHandler.h
@@ -534,7 +534,7 @@ void TClientBubbleHandlerBase<AgentArrayItem>::PostReplicatedAddEntitiesHelper(c
 
 
 		FMassExecutionContext ExecContext(EntityManager);
-		FMassEntityQuery Query;
+		FMassEntityQuery Query(EntityManager.AsShared());
 
 		AddRequirementsForSpawnQuery(Query);
 
diff --git a/Engine/Plugins/Runtime/MassGameplay/Source/MassReplication/Public/MassReplicationFragments.h b/Engine/Plugins/Runtime/MassGameplay/Source/MassReplication/Public/MassReplicationFragments.h
--- a/Engine/Plugins/Runtime/MassGameplay/Source/MassReplication/Public/MassReplicationFragments.h
+++ b/Engine/Plugins/Runtime/MassGameplay/Source/MassReplication/Public/MassReplicationFragments.h
@@ -69,7 +69,7 @@ public:
 	UMassNetworkIDFragmentInitializer();
 
 protected:
-	virtual void ConfigureQueries() override;
+	virtual void ConfigureQueries(const TSharedRef<FMassEntityManager>& EntityManager) override;
 	virtual void Execute(FMassEntityManager& EntityManager, FMassExecutionContext& Context) override;
 
 	FMassEntityQuery EntityQuery;
diff --git a/Engine/Plugins/Runtime/MassGameplay/Source/MassReplication/Public/MassReplicationGridProcessor.h b/Engine/Plugins/Runtime/MassGameplay/Source/MassReplication/Public/MassReplicationGridProcessor.h
--- a/Engine/Plugins/Runtime/MassGameplay/Source/MassReplication/Public/MassReplicationGridProcessor.h
+++ b/Engine/Plugins/Runtime/MassGameplay/Source/MassReplication/Public/MassReplicationGridProcessor.h
@@ -19,7 +19,7 @@ public:
 	UMassReplicationGridProcessor();
 
 protected:
-	virtual void ConfigureQueries() override;
+	virtual void ConfigureQueries(const TSharedRef<FMassEntityManager>& EntityManager) override;
 	virtual void Execute(FMassEntityManager& EntityManager, FMassExecutionContext& Context) override;
 
 private:
@@ -37,7 +37,7 @@ class MASSREPLICATION_API UMassReplicationGridRemoverProcessor : public UMassObs
 	UMassReplicationGridRemoverProcessor();
 
 protected:
-	virtual void ConfigureQueries() override;
+	virtual void ConfigureQueries(const TSharedRef<FMassEntityManager>& EntityManager) override;
 	virtual void Execute(FMassEntityManager& EntityManager, FMassExecutionContext& Context) override;
 
 	FMassEntityQuery EntityQuery;
diff --git a/Engine/Plugins/Runtime/MassGameplay/Source/MassReplication/Public/MassReplicationProcessor.h b/Engine/Plugins/Runtime/MassGameplay/Source/MassReplication/Public/MassReplicationProcessor.h
--- a/Engine/Plugins/Runtime/MassGameplay/Source/MassReplication/Public/MassReplicationProcessor.h
+++ b/Engine/Plugins/Runtime/MassGameplay/Source/MassReplication/Public/MassReplicationProcessor.h
@@ -32,7 +32,7 @@ public:
 	UMassReplicationProcessor();
 
 protected:
-	virtual void ConfigureQueries() override;
+	virtual void ConfigureQueries(const TSharedRef<FMassEntityManager>& EntityManager) override;
 	virtual void InitializeInternal(UObject& Owner, const TSharedRef<FMassEntityManager>& EntityManager) override;
 	virtual void Execute(FMassEntityManager& EntityManager, FMassExecutionContext& Context) override;
 
diff --git a/Engine/Plugins/Runtime/MassGameplay/Source/MassRepresentation/Private/MassDistanceLODProcessor.cpp b/Engine/Plugins/Runtime/MassGameplay/Source/MassRepresentation/Private/MassDistanceLODProcessor.cpp
--- a/Engine/Plugins/Runtime/MassGameplay/Source/MassRepresentation/Private/MassDistanceLODProcessor.cpp
+++ b/Engine/Plugins/Runtime/MassGameplay/Source/MassRepresentation/Private/MassDistanceLODProcessor.cpp
@@ -16,9 +16,9 @@ UMassDistanceLODProcessor::UMassDistanceLODProcessor()
 	ExecutionOrder.ExecuteAfter.Add(UE::Mass::ProcessorGroupNames::LODCollector);
 }
 
-void UMassDistanceLODProcessor::ConfigureQueries()
+void UMassDistanceLODProcessor::ConfigureQueries(const TSharedRef<FMassEntityManager>& EntityManager)
 {
-	FMassEntityQuery BaseQuery;
+	FMassEntityQuery BaseQuery(EntityManager);
 	BaseQuery.AddTagRequirement<FMassDistanceLODProcessorTag>(EMassFragmentPresence::All);
 	BaseQuery.AddRequirement<FMassViewerInfoFragment>(EMassFragmentAccess::ReadOnly);
 	BaseQuery.AddRequirement<FMassRepresentationLODFragment>(EMassFragmentAccess::ReadWrite);
diff --git a/Engine/Plugins/Runtime/MassGameplay/Source/MassRepresentation/Private/MassRepresentationProcessor.cpp b/Engine/Plugins/Runtime/MassGameplay/Source/MassRepresentation/Private/MassRepresentationProcessor.cpp
--- a/Engine/Plugins/Runtime/MassGameplay/Source/MassRepresentation/Private/MassRepresentationProcessor.cpp
+++ b/Engine/Plugins/Runtime/MassGameplay/Source/MassRepresentation/Private/MassRepresentationProcessor.cpp
@@ -36,7 +36,7 @@ UMassRepresentationProcessor::UMassRepresentationProcessor()
 	ExecutionOrder.ExecuteAfter.Add(UE::Mass::ProcessorGroupNames::LOD);
 }
 
-void UMassRepresentationProcessor::ConfigureQueries()
+void UMassRepresentationProcessor::ConfigureQueries(const TSharedRef<FMassEntityManager>& EntityManager)
 {
 	EntityQuery.AddRequirement<FTransformFragment>(EMassFragmentAccess::ReadOnly);
 	EntityQuery.AddRequirement<FMassRepresentationFragment>(EMassFragmentAccess::ReadWrite);
@@ -359,9 +359,9 @@ void UMassVisualizationProcessor::UpdateVisualization(FMassExecutionContext& Con
 	}
 }
 
-void UMassVisualizationProcessor::ConfigureQueries()
+void UMassVisualizationProcessor::ConfigureQueries(const TSharedRef<FMassEntityManager>& EntityManager)
 {
-	Super::ConfigureQueries();
+	Super::ConfigureQueries(EntityManager);
 	EntityQuery.AddChunkRequirement<FMassVisualizationChunkFragment>(EMassFragmentAccess::ReadWrite);
 	EntityQuery.AddTagRequirement<FMassVisualizationProcessorTag>(EMassFragmentPresence::All);
 }
@@ -451,7 +451,7 @@ UMassRepresentationFragmentDestructor::UMassRepresentationFragmentDestructor()
 	bRequiresGameThreadExecution = true; // due to FMassRepresentationSubsystemSharedFragment.RepresentationSubsystem use
 }
 
-void UMassRepresentationFragmentDestructor::ConfigureQueries()
+void UMassRepresentationFragmentDestructor::ConfigureQueries(const TSharedRef<FMassEntityManager>& EntityManager)
 {
 	EntityQuery.AddRequirement<FMassRepresentationFragment>(EMassFragmentAccess::ReadWrite);
 	EntityQuery.AddRequirement<FMassActorFragment>(EMassFragmentAccess::ReadWrite);
diff --git a/Engine/Plugins/Runtime/MassGameplay/Source/MassRepresentation/Private/MassStationaryISMRepresentationFragmentDestructor.cpp b/Engine/Plugins/Runtime/MassGameplay/Source/MassRepresentation/Private/MassStationaryISMRepresentationFragmentDestructor.cpp
--- a/Engine/Plugins/Runtime/MassGameplay/Source/MassRepresentation/Private/MassStationaryISMRepresentationFragmentDestructor.cpp
+++ b/Engine/Plugins/Runtime/MassGameplay/Source/MassRepresentation/Private/MassStationaryISMRepresentationFragmentDestructor.cpp
@@ -17,7 +17,7 @@ UMassStationaryISMRepresentationFragmentDestructor::UMassStationaryISMRepresenta
 	bRequiresGameThreadExecution = true; // not sure about this
 }
 
-void UMassStationaryISMRepresentationFragmentDestructor::ConfigureQueries()
+void UMassStationaryISMRepresentationFragmentDestructor::ConfigureQueries(const TSharedRef<FMassEntityManager>& EntityManager)
 {
 	EntityQuery.AddRequirement<FMassRepresentationFragment>(EMassFragmentAccess::ReadWrite);
 	EntityQuery.AddConstSharedRequirement<FMassRepresentationParameters>();
diff --git a/Engine/Plugins/Runtime/MassGameplay/Source/MassRepresentation/Private/MassStationaryISMSwitcherProcessor.cpp b/Engine/Plugins/Runtime/MassGameplay/Source/MassRepresentation/Private/MassStationaryISMSwitcherProcessor.cpp
--- a/Engine/Plugins/Runtime/MassGameplay/Source/MassRepresentation/Private/MassStationaryISMSwitcherProcessor.cpp
+++ b/Engine/Plugins/Runtime/MassGameplay/Source/MassRepresentation/Private/MassStationaryISMSwitcherProcessor.cpp
@@ -17,7 +17,7 @@ UMassStationaryISMSwitcherProcessor::UMassStationaryISMSwitcherProcessor(const F
 	ExecutionFlags = static_cast<int32>(EProcessorExecutionFlags::AllNetModes);
 }
 
-void UMassStationaryISMSwitcherProcessor::ConfigureQueries()
+void UMassStationaryISMSwitcherProcessor::ConfigureQueries(const TSharedRef<FMassEntityManager>& EntityManager)
 {
 	EntityQuery.AddRequirement<FMassRepresentationLODFragment>(EMassFragmentAccess::ReadOnly);
 	EntityQuery.AddRequirement<FTransformFragment>(EMassFragmentAccess::ReadOnly);
diff --git a/Engine/Plugins/Runtime/MassGameplay/Source/MassRepresentation/Private/MassUpdateISMProcessor.cpp b/Engine/Plugins/Runtime/MassGameplay/Source/MassRepresentation/Private/MassUpdateISMProcessor.cpp
--- a/Engine/Plugins/Runtime/MassGameplay/Source/MassRepresentation/Private/MassUpdateISMProcessor.cpp
+++ b/Engine/Plugins/Runtime/MassGameplay/Source/MassRepresentation/Private/MassUpdateISMProcessor.cpp
@@ -19,7 +19,7 @@ UMassUpdateISMProcessor::UMassUpdateISMProcessor()
 	bRequiresGameThreadExecution = true;
 }
 
-void UMassUpdateISMProcessor::ConfigureQueries()
+void UMassUpdateISMProcessor::ConfigureQueries(const TSharedRef<FMassEntityManager>& EntityManager)
 {
 	EntityQuery.AddRequirement<FTransformFragment>(EMassFragmentAccess::ReadOnly);
 	EntityQuery.AddRequirement<FMassRepresentationFragment>(EMassFragmentAccess::ReadWrite);
diff --git a/Engine/Plugins/Runtime/MassGameplay/Source/MassRepresentation/Private/MassVisualizationLODProcessor.cpp b/Engine/Plugins/Runtime/MassGameplay/Source/MassRepresentation/Private/MassVisualizationLODProcessor.cpp
--- a/Engine/Plugins/Runtime/MassGameplay/Source/MassRepresentation/Private/MassVisualizationLODProcessor.cpp
+++ b/Engine/Plugins/Runtime/MassGameplay/Source/MassRepresentation/Private/MassVisualizationLODProcessor.cpp
@@ -13,9 +13,9 @@ UMassVisualizationLODProcessor::UMassVisualizationLODProcessor()
 	ExecutionOrder.ExecuteAfter.Add(UE::Mass::ProcessorGroupNames::LODCollector);
 }
 
-void UMassVisualizationLODProcessor::ConfigureQueries()
+void UMassVisualizationLODProcessor::ConfigureQueries(const TSharedRef<FMassEntityManager>& EntityManager)
 {
-	FMassEntityQuery BaseQuery;
+	FMassEntityQuery BaseQuery(EntityManager);
 	BaseQuery.AddTagRequirement<FMassVisualizationLODProcessorTag>(EMassFragmentPresence::All);
 	BaseQuery.AddRequirement<FMassViewerInfoFragment>(EMassFragmentAccess::ReadOnly);
 	BaseQuery.AddRequirement<FMassRepresentationLODFragment>(EMassFragmentAccess::ReadWrite);
diff --git a/Engine/Plugins/Runtime/MassGameplay/Source/MassRepresentation/Public/MassDistanceLODProcessor.h b/Engine/Plugins/Runtime/MassGameplay/Source/MassRepresentation/Public/MassDistanceLODProcessor.h
--- a/Engine/Plugins/Runtime/MassGameplay/Source/MassRepresentation/Public/MassDistanceLODProcessor.h
+++ b/Engine/Plugins/Runtime/MassGameplay/Source/MassRepresentation/Public/MassDistanceLODProcessor.h
@@ -32,7 +32,7 @@ public:
 protected:
 
 	/** Configure the owned FMassEntityQuery instances to express processor's requirements */
-	virtual void ConfigureQueries() override;
+	virtual void ConfigureQueries(const TSharedRef<FMassEntityManager>& EntityManager) override;
 
 	/** 
 	 * Execution method for this processor 
diff --git a/Engine/Plugins/Runtime/MassGameplay/Source/MassRepresentation/Public/MassRepresentationProcessor.h b/Engine/Plugins/Runtime/MassGameplay/Source/MassRepresentation/Public/MassRepresentationProcessor.h
--- a/Engine/Plugins/Runtime/MassGameplay/Source/MassRepresentation/Public/MassRepresentationProcessor.h
+++ b/Engine/Plugins/Runtime/MassGameplay/Source/MassRepresentation/Public/MassRepresentationProcessor.h
@@ -49,7 +49,7 @@ public:
 protected:
 
 	/** Configure the owned FMassEntityQuery instances to express processor's requirements */
-	virtual void ConfigureQueries() override;
+	virtual void ConfigureQueries(const TSharedRef<FMassEntityManager>& EntityManager) override;
 
 	/** 
 	 * Execution method for this processor 
@@ -98,7 +98,7 @@ class MASSREPRESENTATION_API UMassVisualizationProcessor : public UMassRepresent
 protected:
 
 	/** Configure the owned FMassEntityQuery instances to express processor's requirements */
-	virtual void ConfigureQueries() override;
+	virtual void ConfigureQueries(const TSharedRef<FMassEntityManager>& EntityManager) override;
 
 	/**
 	 * Execution method for this processor
@@ -141,7 +141,7 @@ public:
 	UMassRepresentationFragmentDestructor();
 
 protected:
-	virtual void ConfigureQueries() override;
+	virtual void ConfigureQueries(const TSharedRef<FMassEntityManager>& EntityManager) override;
 	virtual void Execute(FMassEntityManager& EntityManager, FMassExecutionContext& Context) override;
 
 	FMassEntityQuery EntityQuery;
diff --git a/Engine/Plugins/Runtime/MassGameplay/Source/MassRepresentation/Public/MassStationaryISMRepresentationFragmentDestructor.h b/Engine/Plugins/Runtime/MassGameplay/Source/MassRepresentation/Public/MassStationaryISMRepresentationFragmentDestructor.h
--- a/Engine/Plugins/Runtime/MassGameplay/Source/MassRepresentation/Public/MassStationaryISMRepresentationFragmentDestructor.h
+++ b/Engine/Plugins/Runtime/MassGameplay/Source/MassRepresentation/Public/MassStationaryISMRepresentationFragmentDestructor.h
@@ -17,7 +17,7 @@ public:
 	UMassStationaryISMRepresentationFragmentDestructor();
 
 protected:
-	virtual void ConfigureQueries() override;
+	virtual void ConfigureQueries(const TSharedRef<FMassEntityManager>& EntityManager) override;
 	virtual void Execute(FMassEntityManager& EntityManager, FMassExecutionContext& Context) override;
 
 	FMassEntityQuery EntityQuery;
diff --git a/Engine/Plugins/Runtime/MassGameplay/Source/MassRepresentation/Public/MassStationaryISMSwitcherProcessor.h b/Engine/Plugins/Runtime/MassGameplay/Source/MassRepresentation/Public/MassStationaryISMSwitcherProcessor.h
--- a/Engine/Plugins/Runtime/MassGameplay/Source/MassRepresentation/Public/MassStationaryISMSwitcherProcessor.h
+++ b/Engine/Plugins/Runtime/MassGameplay/Source/MassRepresentation/Public/MassStationaryISMSwitcherProcessor.h
@@ -41,7 +41,7 @@ public:
 
 protected:
 	// UMassProcessor overrides begin
-	virtual void ConfigureQueries() override;
+	virtual void ConfigureQueries(const TSharedRef<FMassEntityManager>& EntityManager) override;
 	// UMassProcessor overrides end
 	// 
 	// UMassRepresentationProcessor overrides begin
diff --git a/Engine/Plugins/Runtime/MassGameplay/Source/MassRepresentation/Public/MassUpdateISMProcessor.h b/Engine/Plugins/Runtime/MassGameplay/Source/MassRepresentation/Public/MassUpdateISMProcessor.h
--- a/Engine/Plugins/Runtime/MassGameplay/Source/MassRepresentation/Public/MassUpdateISMProcessor.h
+++ b/Engine/Plugins/Runtime/MassGameplay/Source/MassRepresentation/Public/MassUpdateISMProcessor.h
@@ -26,7 +26,7 @@ public:
 protected:
 
 	/** Configure the owned FMassEntityQuery instances to express processor's requirements */
-	virtual void ConfigureQueries() override;
+	virtual void ConfigureQueries(const TSharedRef<FMassEntityManager>& EntityManager) override;
 
 	/**
 	 * Execution method for this processor
diff --git a/Engine/Plugins/Runtime/MassGameplay/Source/MassRepresentation/Public/MassVisualizationLODProcessor.h b/Engine/Plugins/Runtime/MassGameplay/Source/MassRepresentation/Public/MassVisualizationLODProcessor.h
--- a/Engine/Plugins/Runtime/MassGameplay/Source/MassRepresentation/Public/MassVisualizationLODProcessor.h
+++ b/Engine/Plugins/Runtime/MassGameplay/Source/MassRepresentation/Public/MassVisualizationLODProcessor.h
@@ -32,7 +32,7 @@ public:
 protected:
 
 	/** Configure the owned FMassEntityQuery instances to express processor's requirements */
-	virtual void ConfigureQueries() override;
+	virtual void ConfigureQueries(const TSharedRef<FMassEntityManager>& EntityManager) override;
 
 	/** 
 	 * Execution method for this processor 
diff --git a/Engine/Plugins/Runtime/MassGameplay/Source/MassSignals/Public/MassSignalProcessorBase.h b/Engine/Plugins/Runtime/MassGameplay/Source/MassSignals/Public/MassSignalProcessorBase.h
--- a/Engine/Plugins/Runtime/MassGameplay/Source/MassSignals/Public/MassSignalProcessorBase.h
+++ b/Engine/Plugins/Runtime/MassGameplay/Source/MassSignals/Public/MassSignalProcessorBase.h
@@ -27,7 +27,7 @@ protected:
 	virtual void BeginDestroy() override;
 
 	/** Configure the owned FMassEntityQuery instances to express processor queries requirements */
-	virtual void ConfigureQueries() override {}
+	virtual void ConfigureQueries(const TSharedRef<FMassEntityManager>&) override {}
 
 	/**
 	 * Actual method that derived class needs to implement to act on a signal that is raised for that frame
diff --git a/Engine/Plugins/Runtime/MassGameplay/Source/MassSmartObjects/Private/MassSmartObjectProcessor.cpp b/Engine/Plugins/Runtime/MassGameplay/Source/MassSmartObjects/Private/MassSmartObjectProcessor.cpp
--- a/Engine/Plugins/Runtime/MassGameplay/Source/MassSmartObjects/Private/MassSmartObjectProcessor.cpp
+++ b/Engine/Plugins/Runtime/MassGameplay/Source/MassSmartObjects/Private/MassSmartObjectProcessor.cpp
@@ -24,7 +24,7 @@
 //----------------------------------------------------------------------//
 // UMassSmartObjectCandidatesFinderProcessor
 //----------------------------------------------------------------------//
-void UMassSmartObjectCandidatesFinderProcessor::ConfigureQueries()
+void UMassSmartObjectCandidatesFinderProcessor::ConfigureQueries(const TSharedRef<FMassEntityManager>& EntityManager)
 {
 	WorldRequestQuery.AddRequirement<FMassSmartObjectWorldLocationRequestFragment>(EMassFragmentAccess::ReadOnly);
 	WorldRequestQuery.AddRequirement<FMassSmartObjectRequestResultFragment>(EMassFragmentAccess::ReadWrite);
@@ -304,7 +304,7 @@ void UMassSmartObjectCandidatesFinderProcessor::Execute(FMassEntityManager& Enti
 //----------------------------------------------------------------------//
 // UMassSmartObjectTimedBehaviorProcessor
 //----------------------------------------------------------------------//
-void UMassSmartObjectTimedBehaviorProcessor::ConfigureQueries()
+void UMassSmartObjectTimedBehaviorProcessor::ConfigureQueries(const TSharedRef<FMassEntityManager>& EntityManager)
 {
 	EntityQuery.AddRequirement<FMassSmartObjectUserFragment>(EMassFragmentAccess::ReadWrite);
 	EntityQuery.AddRequirement<FMassSmartObjectTimedBehaviorFragment>(EMassFragmentAccess::ReadWrite);
@@ -392,7 +392,7 @@ UMassSmartObjectUserFragmentDeinitializer::UMassSmartObjectUserFragmentDeinitial
 	ExecutionFlags = (int32)EProcessorExecutionFlags::AllNetModes;
 }
 
-void UMassSmartObjectUserFragmentDeinitializer::ConfigureQueries()
+void UMassSmartObjectUserFragmentDeinitializer::ConfigureQueries(const TSharedRef<FMassEntityManager>& EntityManager)
 {
 	EntityQuery.AddRequirement<FMassSmartObjectUserFragment>(EMassFragmentAccess::ReadWrite);
 	EntityQuery.AddSubsystemRequirement<USmartObjectSubsystem>(EMassFragmentAccess::ReadWrite);
diff --git a/Engine/Plugins/Runtime/MassGameplay/Source/MassSmartObjects/Private/MassSmartObjectRegistration.cpp b/Engine/Plugins/Runtime/MassGameplay/Source/MassSmartObjects/Private/MassSmartObjectRegistration.cpp
--- a/Engine/Plugins/Runtime/MassGameplay/Source/MassSmartObjects/Private/MassSmartObjectRegistration.cpp
+++ b/Engine/Plugins/Runtime/MassGameplay/Source/MassSmartObjects/Private/MassSmartObjectRegistration.cpp
@@ -24,7 +24,7 @@ UMassSmartObjectInitializerBase::UMassSmartObjectInitializerBase()
 	Operation = EMassObservedOperation::Add;
 }
 
-void UMassSmartObjectInitializerBase::ConfigureQueries()
+void UMassSmartObjectInitializerBase::ConfigureQueries(const TSharedRef<FMassEntityManager>& EntityManager)
 {
 	ProcessorRequirements.AddSubsystemRequirement<UMassSignalSubsystem>(EMassFragmentAccess::ReadWrite);
 	EntityQuery.AddTagRequirement<FMassInActiveSmartObjectsRangeTag>(EMassFragmentPresence::All);
@@ -82,7 +82,7 @@ UMassSmartObjectDeinitializerBase::UMassSmartObjectDeinitializerBase()
 	Operation = EMassObservedOperation::Remove;
 }
 
-void UMassSmartObjectDeinitializerBase::ConfigureQueries()
+void UMassSmartObjectDeinitializerBase::ConfigureQueries(const TSharedRef<FMassEntityManager>& EntityManager)
 {
 	ProcessorRequirements.AddSubsystemRequirement<UMassSignalSubsystem>(EMassFragmentAccess::ReadWrite);
 }
@@ -157,7 +157,7 @@ UMassActiveSmartObjectSignalProcessor::UMassActiveSmartObjectSignalProcessor()
 	bRequiresGameThreadExecution = true;
 }
 
-void UMassActiveSmartObjectSignalProcessor::ConfigureQueries()
+void UMassActiveSmartObjectSignalProcessor::ConfigureQueries(const TSharedRef<FMassEntityManager>& EntityManager)
 {
 	InsideSmartObjectActiveRangeQuery.AddSubsystemRequirement<USmartObjectSubsystem>(EMassFragmentAccess::ReadWrite);
 	InsideSmartObjectActiveRangeQuery.AddRequirement<FTransformFragment>(EMassFragmentAccess::ReadOnly);
diff --git a/Engine/Plugins/Runtime/MassGameplay/Source/MassSmartObjects/Public/MassSmartObjectProcessor.h b/Engine/Plugins/Runtime/MassGameplay/Source/MassSmartObjects/Public/MassSmartObjectProcessor.h
--- a/Engine/Plugins/Runtime/MassGameplay/Source/MassSmartObjects/Public/MassSmartObjectProcessor.h
+++ b/Engine/Plugins/Runtime/MassGameplay/Source/MassSmartObjects/Public/MassSmartObjectProcessor.h
@@ -19,7 +19,7 @@ public:
 	UMassSmartObjectCandidatesFinderProcessor();
 
 protected:
-	virtual void ConfigureQueries() override;
+	virtual void ConfigureQueries(const TSharedRef<FMassEntityManager>& EntityManager) override;
 	virtual void Execute(FMassEntityManager& EntityManager, FMassExecutionContext& Context) override;
 
 	/** Extents used to perform the spatial query in the octree for world location queries. */
@@ -43,7 +43,7 @@ public:
 
 protected:
 	virtual void Execute(FMassEntityManager& EntityManager, FMassExecutionContext& Context) override;
-	virtual void ConfigureQueries() override;
+	virtual void ConfigureQueries(const TSharedRef<FMassEntityManager>& EntityManager) override;
 
 	FMassEntityQuery EntityQuery;
 };
@@ -57,7 +57,7 @@ class MASSSMARTOBJECTS_API UMassSmartObjectUserFragmentDeinitializer : public UM
 	UMassSmartObjectUserFragmentDeinitializer();
 
 protected:
-	virtual void ConfigureQueries() override;
+	virtual void ConfigureQueries(const TSharedRef<FMassEntityManager>& EntityManager) override;
 	virtual void Execute(FMassEntityManager& EntityManager, FMassExecutionContext& Context) override;
 
 	FMassEntityQuery EntityQuery;
diff --git a/Engine/Plugins/Runtime/MassGameplay/Source/MassSmartObjects/Public/MassSmartObjectRegistration.h b/Engine/Plugins/Runtime/MassGameplay/Source/MassSmartObjects/Public/MassSmartObjectRegistration.h
--- a/Engine/Plugins/Runtime/MassGameplay/Source/MassSmartObjects/Public/MassSmartObjectRegistration.h
+++ b/Engine/Plugins/Runtime/MassGameplay/Source/MassSmartObjects/Public/MassSmartObjectRegistration.h
@@ -43,7 +43,7 @@ public:
 	UMassSmartObjectInitializerBase();
 
 protected:
-	virtual void ConfigureQueries() override;
+	virtual void ConfigureQueries(const TSharedRef<FMassEntityManager>& EntityManager) override;
 	virtual void Execute(FMassEntityManager& EntityManager, FMassExecutionContext& Context) override;
 
 	FMassEntityQuery EntityQuery;
@@ -64,7 +64,7 @@ public:
 	UMassSmartObjectDeinitializerBase();
 
 protected:
-	virtual void ConfigureQueries() override;
+	virtual void ConfigureQueries(const TSharedRef<FMassEntityManager>& EntityManager) override;
 	virtual void Execute(FMassEntityManager& EntityManager, FMassExecutionContext& Context) override;
 
 	FMassEntityQuery EntityQuery;
@@ -134,7 +134,7 @@ public:
 	UMassActiveSmartObjectSignalProcessor();
 
 protected:
-	virtual void ConfigureQueries() override;
+	virtual void ConfigureQueries(const TSharedRef<FMassEntityManager>& EntityManager) override;
 	virtual void InitializeInternal(UObject& Owner, const TSharedRef<FMassEntityManager>& EntityManager) override;
 	virtual void SignalEntities(FMassEntityManager& EntityManager, FMassExecutionContext& Context, FMassSignalNameLookup& /*Unused*/) override;
 
diff --git a/Engine/Plugins/Runtime/MassGameplay/Source/MassSpawner/Private/MassSpawnLocationProcessor.cpp b/Engine/Plugins/Runtime/MassGameplay/Source/MassSpawner/Private/MassSpawnLocationProcessor.cpp
--- a/Engine/Plugins/Runtime/MassGameplay/Source/MassSpawner/Private/MassSpawnLocationProcessor.cpp
+++ b/Engine/Plugins/Runtime/MassGameplay/Source/MassSpawner/Private/MassSpawnLocationProcessor.cpp
@@ -19,7 +19,7 @@ UMassSpawnLocationProcessor::UMassSpawnLocationProcessor()
 	RandomStream.Initialize(UE::Mass::Utils::GenerateRandomSeed());
 }
 
-void UMassSpawnLocationProcessor::ConfigureQueries()
+void UMassSpawnLocationProcessor::ConfigureQueries(const TSharedRef<FMassEntityManager>& EntityManager)
 {
 	EntityQuery.AddRequirement<FTransformFragment>(EMassFragmentAccess::ReadWrite);
 }
diff --git a/Engine/Plugins/Runtime/MassGameplay/Source/MassSpawner/Public/MassSpawnLocationProcessor.h b/Engine/Plugins/Runtime/MassGameplay/Source/MassSpawner/Public/MassSpawnLocationProcessor.h
--- a/Engine/Plugins/Runtime/MassGameplay/Source/MassSpawner/Public/MassSpawnLocationProcessor.h
+++ b/Engine/Plugins/Runtime/MassGameplay/Source/MassSpawner/Public/MassSpawnLocationProcessor.h
@@ -14,7 +14,7 @@ public:
 	UMassSpawnLocationProcessor();
 
 protected:
-	virtual void ConfigureQueries() override;
+	virtual void ConfigureQueries(const TSharedRef<FMassEntityManager>& EntityManager) override;
 	virtual void Execute(FMassEntityManager& EntityManager, FMassExecutionContext& Context) override;
 
 	FMassEntityQuery EntityQuery;
diff --git a/Engine/Source/Developer/MassEntityTestSuite/Private/MassChunkCollectionTest.cpp b/Engine/Source/Developer/MassEntityTestSuite/Private/MassChunkCollectionTest.cpp
--- a/Engine/Source/Developer/MassEntityTestSuite/Private/MassChunkCollectionTest.cpp
+++ b/Engine/Source/Developer/MassEntityTestSuite/Private/MassChunkCollectionTest.cpp
@@ -402,7 +402,7 @@ struct FEntityCollection_ZeroLength_Destroy : FEntityTestBase
 
 		AITEST_EQUAL("Entity creation generated expected number of entities", EntityCount, FMassArchetypeHelper::ArchetypeDataFromHandleChecked(SourceArchetype).GetNumEntities());
 
-		FMassEntityQuery Query;
+		FMassEntityQuery Query(EntityManager);
 		Query.CacheArchetypes(*EntityManager);
 		FMassArchetypeEntityCollection Collection(SourceArchetype);
 		AITEST_EQUAL("Created collection has expected number of ranges (matching number of chunks created"
@@ -432,7 +432,7 @@ struct FEntityCollection_ZeroLength_Execute : FEntityTestBase
 		AITEST_EQUAL("Entity creation generated expected number of entities", EntityCount, FMassArchetypeHelper::ArchetypeDataFromHandleChecked(SourceArchetype).GetNumEntities());
 
 		int32 ProcessedEntitiesCount = 0;
-		UMassTestProcessorBase* Processor = NewObject<UMassTestProcessorBase>();
+		UMassTestProcessorBase* Processor = NewTestProcessor<UMassTestProcessorBase>(EntityManager);
 		Processor->ForEachEntityChunkExecutionFunction = [&ProcessedEntitiesCount](FMassExecutionContext& Context)
 		{
 			ProcessedEntitiesCount += Context.GetNumEntities();
diff --git a/Engine/Source/Developer/MassEntityTestSuite/Private/MassCommandsTest.cpp b/Engine/Source/Developer/MassEntityTestSuite/Private/MassCommandsTest.cpp
--- a/Engine/Source/Developer/MassEntityTestSuite/Private/MassCommandsTest.cpp
+++ b/Engine/Source/Developer/MassEntityTestSuite/Private/MassCommandsTest.cpp
@@ -264,7 +264,7 @@ struct FCommands_PushWhileFlushing : FEntityTestBase
 		}
 
 
-		UMassTestProcessorBase* ObserverProcessor = NewObject<UMassTestProcessorBase>();
+		UMassTestProcessorBase* ObserverProcessor = NewTestProcessor<UMassTestProcessorBase>(EntityManager);
 		ObserverProcessor->ForEachEntityChunkExecutionFunction = [](FMassExecutionContext& Context)
 			{
 				for (const FMassEntityHandle& EntityHandle : Context.GetEntities())
diff --git a/Engine/Source/Developer/MassEntityTestSuite/Private/MassComplexTest.cpp b/Engine/Source/Developer/MassEntityTestSuite/Private/MassComplexTest.cpp
--- a/Engine/Source/Developer/MassEntityTestSuite/Private/MassComplexTest.cpp
+++ b/Engine/Source/Developer/MassEntityTestSuite/Private/MassComplexTest.cpp
@@ -42,7 +42,7 @@ struct FComplex_AddingFragmentInstancesToDiscontinouousEntitiesCollection : FEnt
 		EntityManager->BatchCreateEntities(IntsArchetype, EntitiesCount, EntitiesCreated);
 
 		FMassExecutionContext ExecContext(*EntityManager.Get());
-		FMassEntityQuery Query({ FTestFragment_Int::StaticStruct() });
+		FMassEntityQuery Query(EntityManager.ToSharedRef(), { FTestFragment_Int::StaticStruct() });
 		Query.ForEachEntityChunk(*EntityManager, ExecContext, [](FMassExecutionContext& Context)
 			{
 				TArrayView<FTestFragment_Int> Ints = Context.GetMutableFragmentView<FTestFragment_Int>();
diff --git a/Engine/Source/Developer/MassEntityTestSuite/Private/MassCompositeProcessorTest.cpp b/Engine/Source/Developer/MassEntityTestSuite/Private/MassCompositeProcessorTest.cpp
--- a/Engine/Source/Developer/MassEntityTestSuite/Private/MassCompositeProcessorTest.cpp
+++ b/Engine/Source/Developer/MassEntityTestSuite/Private/MassCompositeProcessorTest.cpp
@@ -52,7 +52,7 @@ struct FCompositeProcessorTest_NoWork : FEntityTestBase
 			TArray<UMassProcessor*> Processors;
 			for (int i = 0; i < 3; ++i)
 			{
-				UMassTestProcessorBase* Processor = NewObject<UMassTestProcessorBase>();
+				UMassTestProcessorBase* Processor = NewTestProcessor<UMassTestProcessorBase>(EntityManager);
 				Processor->ForEachEntityChunkExecutionFunction = QueryExecFunction;
 				// need to set up some requirements to make EntityQuery valid
 				Processor->EntityQuery.AddRequirement<FTestFragment_Float>(EMassFragmentAccess::ReadOnly);
@@ -77,7 +77,7 @@ struct FCompositeProcessorTest_NoWork : FEntityTestBase
 			TArray<TObjectPtr<UMassProcessor>> Processors;
 			for (int i = 0; i < 3; ++i)
 			{
-				UMassTestProcessorBase* Processor = NewObject<UMassTestProcessorBase>();
+				UMassTestProcessorBase* Processor = NewTestProcessor<UMassTestProcessorBase>(EntityManager);
 				Processor->ForEachEntityChunkExecutionFunction = QueryExecFunction;
 				Processor->EntityQuery.AddRequirement<FTestFragment_Int>(EMassFragmentAccess::ReadWrite);
 				Processors.Add(Processor);
@@ -111,7 +111,7 @@ struct FCompositeProcessorTest_MultipleSubProcessors : FEntityTestBase
 			TArray<UMassProcessor*> Processors;
 			for (int i = 0; i < 3; ++i)
 			{
-				UMassTestProcessorBase* Processor = NewObject<UMassTestProcessorBase>();
+				UMassTestProcessorBase* Processor = NewTestProcessor<UMassTestProcessorBase>(EntityManager);
 				Processor->EntityQuery.AddRequirement<FTestFragment_Int>(EMassFragmentAccess::ReadOnly);
 				int Value = (int)FMath::Pow(10.f, float(i));
 				Processor->ForEachEntityChunkExecutionFunction = [&Result, Value](FMassExecutionContext& Context)
diff --git a/Engine/Source/Developer/MassEntityTestSuite/Private/MassDependencyTest.cpp b/Engine/Source/Developer/MassEntityTestSuite/Private/MassDependencyTest.cpp
--- a/Engine/Source/Developer/MassEntityTestSuite/Private/MassDependencyTest.cpp
+++ b/Engine/Source/Developer/MassEntityTestSuite/Private/MassDependencyTest.cpp
@@ -21,9 +21,11 @@ struct FDependencySolverBase : FAITestBase
 {
 	TArray<UMassTestProcessorBase*> Processors;
 	TArray<FMassProcessorOrderInfo> Result;
+	TSharedPtr<FMassEntityManager> EntityManager;
 	
 	virtual bool SetUp() override
 	{
+		EntityManager = MakeShareable(new FMassEntityManager());
 		Processors.Reset();
 		return true;
 	}
@@ -42,7 +44,8 @@ struct FTrivialDependency : FDependencySolverBase
 
 	virtual bool SetUp() override
 	{
-		Proc = Processors.Add_GetRef(NewObject<UMassTestProcessor_A>());
+		FDependencySolverBase::SetUp();
+		Proc = Processors.Add_GetRef(NewTestProcessor<UMassTestProcessor_A>(EntityManager));
 		return true;
 	}
 
@@ -62,17 +65,18 @@ struct FSimpleDependency : FDependencySolverBase
 {
 	virtual bool SetUp() override
 	{
+		FDependencySolverBase::SetUp();
 		{
-			UMassTestProcessorBase* Proc = Processors.Add_GetRef(NewObject<UMassTestProcessor_A>());
+			UMassTestProcessorBase* Proc = Processors.Add_GetRef(NewTestProcessor<UMassTestProcessor_A>(EntityManager));
 			Proc->GetMutableExecutionOrder().ExecuteAfter.Add(GetProcessorName<UMassTestProcessor_C>());
 		}
 
 		{
-			UMassTestProcessorBase* Proc = Processors.Add_GetRef(NewObject<UMassTestProcessor_B>());
+			UMassTestProcessorBase* Proc = Processors.Add_GetRef(NewTestProcessor<UMassTestProcessor_B>(EntityManager));
 			Proc->GetMutableExecutionOrder().ExecuteAfter.Add(GetProcessorName<UMassTestProcessor_A>());
 		}
 
-		Processors.Add(NewObject<UMassTestProcessor_C>());
+		Processors.Add(NewTestProcessor<UMassTestProcessor_C>(EntityManager));
 
 		return true;
 	}
@@ -94,17 +98,18 @@ struct FMissingDependency : FDependencySolverBase
 {
 	virtual bool SetUp() override
 	{
+		FDependencySolverBase::SetUp();
 		{
-			UMassTestProcessorBase* Proc = Processors.Add_GetRef(NewObject<UMassTestProcessor_A>());
+			UMassTestProcessorBase* Proc = Processors.Add_GetRef(NewTestProcessor<UMassTestProcessor_A>(EntityManager));
 			Proc->GetMutableExecutionOrder().ExecuteAfter.Add(TEXT("NonExistingDependency"));
 		}
 		{
-			UMassTestProcessorBase* Proc = Processors.Add_GetRef(NewObject<UMassTestProcessor_B>());
+			UMassTestProcessorBase* Proc = Processors.Add_GetRef(NewTestProcessor<UMassTestProcessor_B>(EntityManager));
 			Proc->GetMutableExecutionOrder().ExecuteBefore.Add(TEXT("NonExistingDependency2"));
 			Proc->GetMutableExecutionOrder().ExecuteAfter.Add(GetProcessorName<UMassTestProcessor_C>());
 		}
 		{
-			UMassTestProcessorBase* Proc = Processors.Add_GetRef(NewObject<UMassTestProcessor_C>());
+			UMassTestProcessorBase* Proc = Processors.Add_GetRef(NewTestProcessor<UMassTestProcessor_C>(EntityManager));
 		}
 		return true;
 	}
@@ -130,14 +135,15 @@ struct FDeepGroup : FDependencySolverBase
 {
 	virtual bool SetUp() override
 	{
+		FDependencySolverBase::SetUp();
 		{
-			UMassTestProcessorBase* Proc = Processors.Add_GetRef(NewObject<UMassTestProcessor_A>());
+			UMassTestProcessorBase* Proc = Processors.Add_GetRef(NewTestProcessor<UMassTestProcessor_A>(EntityManager));
 			Proc->GetMutableExecutionOrder().ExecuteAfter.Add(TEXT("W.X.Y.Z"));
 			Proc->GetMutableExecutionOrder().ExecuteInGroup = TEXT("P.Q.R");
 		}
 
 		{
-			UMassTestProcessorBase* Proc = Processors.Add_GetRef(NewObject<UMassTestProcessor_B>());
+			UMassTestProcessorBase* Proc = Processors.Add_GetRef(NewTestProcessor<UMassTestProcessor_B>(EntityManager));
 			Proc->GetMutableExecutionOrder().ExecuteInGroup = TEXT("W.X.Y.Z");
 		}
 
@@ -169,25 +175,26 @@ struct FComplexScenario : FDependencySolverBase
 {
 	virtual bool SetUp() override
 	{
+		FDependencySolverBase::SetUp();
 		{
-			UMassTestProcessorBase* Proc = Processors.Add_GetRef(NewObject<UMassTestProcessor_A>());
+			UMassTestProcessorBase* Proc = Processors.Add_GetRef(NewTestProcessor<UMassTestProcessor_A>(EntityManager));
 			Proc->GetMutableExecutionOrder().ExecuteInGroup = TEXT("X.Z");
 			Proc->GetMutableExecutionOrder().ExecuteAfter.Add(TEXT("X.Y"));
 			Proc->GetMutableExecutionOrder().ExecuteAfter.Add(UMassTestProcessor_E::StaticClass()->GetFName());
 		}
 
 		{
-			UMassTestProcessorBase* Proc = Processors.Add_GetRef(NewObject<UMassTestProcessor_B>());
+			UMassTestProcessorBase* Proc = Processors.Add_GetRef(NewTestProcessor<UMassTestProcessor_B>(EntityManager));
 			Proc->GetMutableExecutionOrder().ExecuteInGroup = TEXT("X.Y");
 		}
 
 		{
-			UMassTestProcessorBase* Proc = Processors.Add_GetRef(NewObject<UMassTestProcessor_C>());
+			UMassTestProcessorBase* Proc = Processors.Add_GetRef(NewTestProcessor<UMassTestProcessor_C>(EntityManager));
 			Proc->GetMutableExecutionOrder().ExecuteInGroup = TEXT("X.Y");
 		}
 
 		{
-			UMassTestProcessorBase* Proc = Processors.Add_GetRef(NewObject<UMassTestProcessor_D>());
+			UMassTestProcessorBase* Proc = Processors.Add_GetRef(NewTestProcessor<UMassTestProcessor_D>(EntityManager));
 			Proc->GetMutableExecutionOrder().ExecuteBefore.Add(UMassTestProcessor_A::StaticClass()->GetFName());
 			Proc->GetMutableExecutionOrder().ExecuteBefore.Add(TEXT("X.Y"));
 		}
@@ -278,9 +285,9 @@ struct FBadInput_MixedNulls : FDependencySolverBase
 	{
 		Processors.Reset();
 		Processors.Add(nullptr);
-		Processors.Add(NewObject<UMassTestProcessor_A>());
+		Processors.Add(NewTestProcessor<UMassTestProcessor_A>(EntityManager));
 		Processors.Add(nullptr);
-		Processors.Add(NewObject<UMassTestProcessor_B>());
+		Processors.Add(NewTestProcessor<UMassTestProcessor_B>(EntityManager));
 		{
 			AITEST_SCOPED_CHECK(TEXT("nullptr found in Processors"), 2);
 			Solve();
@@ -296,9 +303,9 @@ struct FBadInput_Duplicates : FDependencySolverBase
 	virtual bool InstantTest() override
 	{
 		Processors.Reset();
-		Processors.Add(NewObject<UMassTestProcessor_A>());
-		Processors.Add(NewObject<UMassTestProcessor_A>());
-		Processors.Add(NewObject<UMassTestProcessor_A>());
+		Processors.Add(NewTestProcessor<UMassTestProcessor_A>(EntityManager));
+		Processors.Add(NewTestProcessor<UMassTestProcessor_A>(EntityManager));
+		Processors.Add(NewTestProcessor<UMassTestProcessor_A>(EntityManager));
 		{
 			AITEST_SCOPED_CHECK(TEXT("already registered. Duplicates are not supported"), 2);
 			Solve();
@@ -342,23 +349,30 @@ struct FCircularDependency : FDependencySolverBase
 {
 	virtual bool SetUp() override
 	{
+		FDependencySolverBase::SetUp();
+		TSharedRef<FMassEntityManager> EntityManagerRef = EntityManager.ToSharedRef();
+
 		{
-			UMassTestProcessorBase* Proc = Processors.Add_GetRef(NewObject<UMassTestProcessor_A>());
+			UMassTestProcessorBase* Proc = Processors.Add_GetRef(NewTestProcessor<UMassTestProcessor_A>(EntityManager));
+			Proc->EntityQuery.Initialize(EntityManagerRef);
 			Proc->EntityQuery.AddRequirement<FTestFragment_Int>(EMassFragmentAccess::ReadWrite);
 			Proc->GetMutableExecutionOrder().ExecuteAfter.Add(GetProcessorName<UMassTestProcessor_D>());
 		}
 		{
-			UMassTestProcessorBase* Proc = Processors.Add_GetRef(NewObject<UMassTestProcessor_B>());
+			UMassTestProcessorBase* Proc = Processors.Add_GetRef(NewTestProcessor<UMassTestProcessor_B>(EntityManager));
+			Proc->EntityQuery.Initialize(EntityManagerRef);
 			Proc->EntityQuery.AddRequirement<FTestFragment_Int>(EMassFragmentAccess::ReadWrite);
 			Proc->GetMutableExecutionOrder().ExecuteAfter.Add(GetProcessorName<UMassTestProcessor_A>());
 		}
 		{
-			UMassTestProcessorBase* Proc = Processors.Add_GetRef(NewObject<UMassTestProcessor_C>());
+			UMassTestProcessorBase* Proc = Processors.Add_GetRef(NewTestProcessor<UMassTestProcessor_C>(EntityManager));
+			Proc->EntityQuery.Initialize(EntityManagerRef);
 			Proc->EntityQuery.AddRequirement<FTestFragment_Int>(EMassFragmentAccess::ReadWrite);
 			Proc->GetMutableExecutionOrder().ExecuteAfter.Add(GetProcessorName<UMassTestProcessor_B>());
 		}
 		{
-			UMassTestProcessorBase* Proc = Processors.Add_GetRef(NewObject<UMassTestProcessor_D>());
+			UMassTestProcessorBase* Proc = Processors.Add_GetRef(NewTestProcessor<UMassTestProcessor_D>(EntityManager));
+			Proc->EntityQuery.Initialize(EntityManagerRef);
 			Proc->EntityQuery.AddRequirement<FTestFragment_Int>(EMassFragmentAccess::ReadWrite);
 			Proc->GetMutableExecutionOrder().ExecuteAfter.Add(GetProcessorName<UMassTestProcessor_C>());
 		}
diff --git a/Engine/Source/Developer/MassEntityTestSuite/Private/MassDynamicProcessorsTest.cpp b/Engine/Source/Developer/MassEntityTestSuite/Private/MassDynamicProcessorsTest.cpp
--- a/Engine/Source/Developer/MassEntityTestSuite/Private/MassDynamicProcessorsTest.cpp
+++ b/Engine/Source/Developer/MassEntityTestSuite/Private/MassDynamicProcessorsTest.cpp
@@ -38,7 +38,7 @@ struct FDynamicProcessorsAddTrivial : FProcessingPhasesTestBase
 		if (TickIndex == AddDynamicProcessorOnTickIndex)
 		{
 			// on second tick we're adding a dynamic processor
-			DynamicProcessor = NewObject<UMassTestProcessorBase>(World);
+			DynamicProcessor = NewTestProcessor<UMassTestProcessorBase>(EntityManager);
 			DynamicProcessor->ExecutionFunction = [this](FMassEntityManager& InEntityManager, FMassExecutionContext& Context)
 			{
 				++NumberOfTimesTicked;
@@ -183,7 +183,7 @@ struct FDynamicProcessorMultipleInstances : FProcessingPhasesTestBase
 
 			for (int i = 0; i < NumberOfProcessorsToInstantiate; ++i)
 			{
-				UMassTestProcessorBase* DynamicProcessor = NewObject<UMassTestProcessorBase>(World);
+				UMassTestProcessorBase* DynamicProcessor = NewTestProcessor<UMassTestProcessorBase>(EntityManager);
 				DynamicProcessor->SetShouldAllowMultipleInstances(true);
 				DynamicProcessor->ExecutionFunction = [this, i](FMassEntityManager& InEntityManager, FMassExecutionContext& Context)
 				{
diff --git a/Engine/Source/Developer/MassEntityTestSuite/Private/MassEntityStorageTest.cpp b/Engine/Source/Developer/MassEntityTestSuite/Private/MassEntityStorageTest.cpp
--- a/Engine/Source/Developer/MassEntityTestSuite/Private/MassEntityStorageTest.cpp
+++ b/Engine/Source/Developer/MassEntityTestSuite/Private/MassEntityStorageTest.cpp
@@ -99,7 +99,7 @@ struct FEntityStorageTestBase : FAITestBase
 		int32 Index = 0;
 		for (const EEntityManagerMode Mode : ManagersToCreate)
 		{
-			TSharedPtr<FMassEntityManager> LocalEntityManager = MakeShareable(new FMassEntityManager(nullptr));
+			TSharedPtr<FMassEntityManager> LocalEntityManager = MakeShareable(new FMassEntityManager());
 			LocalEntityManager->SetDebugName(FString::Printf(TEXT("TestEntityManager_%d"), Index));
 			++Index;
 
diff --git a/Engine/Source/Developer/MassEntityTestSuite/Private/MassEntityTestFarmPlot.cpp b/Engine/Source/Developer/MassEntityTestSuite/Private/MassEntityTestFarmPlot.cpp
--- a/Engine/Source/Developer/MassEntityTestSuite/Private/MassEntityTestFarmPlot.cpp
+++ b/Engine/Source/Developer/MassEntityTestSuite/Private/MassEntityTestFarmPlot.cpp
@@ -25,7 +25,7 @@ UFarmProcessorBase::UFarmProcessorBase()
 //-----------------------------------------------------------------------------
 // UFarmWaterProcessor
 //-----------------------------------------------------------------------------
-void UFarmWaterProcessor::ConfigureQueries()
+void UFarmWaterProcessor::ConfigureQueries(const TSharedRef<FMassEntityManager>& EntityManager)
 {
 	EntityQuery.AddRequirement<FFarmWaterFragment>(EMassFragmentAccess::ReadWrite);
 }
@@ -48,7 +48,7 @@ void UFarmWaterProcessor::Execute(FMassEntityManager& EntityManager, FMassExecut
 //-----------------------------------------------------------------------------
 // UFarmHarvestTimerSystem_Flowers
 //-----------------------------------------------------------------------------
-void UFarmHarvestTimerSystem_Flowers::ConfigureQueries()
+void UFarmHarvestTimerSystem_Flowers::ConfigureQueries(const TSharedRef<FMassEntityManager>& EntityManager)
 {
 	EntityQuery.AddRequirement<FHarvestTimerFragment>(EMassFragmentAccess::ReadWrite);
 	EntityQuery.AddRequirement<FFarmWaterFragment>(EMassFragmentAccess::ReadOnly);
@@ -84,7 +84,7 @@ void UFarmHarvestTimerSystem_Flowers::Execute(FMassEntityManager& EntityManager,
 //-----------------------------------------------------------------------------
 // UFarmHarvestTimerSystem_Crops
 //-----------------------------------------------------------------------------
-void UFarmHarvestTimerSystem_Crops::ConfigureQueries()
+void UFarmHarvestTimerSystem_Crops::ConfigureQueries(const TSharedRef<FMassEntityManager>& EntityManager)
 {
 	EntityQuery.AddRequirement<FHarvestTimerFragment>(EMassFragmentAccess::ReadWrite);
 	EntityQuery.AddRequirement<FFarmWaterFragment>(EMassFragmentAccess::ReadOnly);
@@ -113,7 +113,7 @@ void UFarmHarvestTimerSystem_Crops::Execute(FMassEntityManager & EntityManager,
 //-----------------------------------------------------------------------------
 // UFarmHarvestTimerExpired
 //-----------------------------------------------------------------------------
-void UFarmHarvestTimerExpired::ConfigureQueries()
+void UFarmHarvestTimerExpired::ConfigureQueries(const TSharedRef<FMassEntityManager>& EntityManager)
 {
 	EntityQuery.AddRequirement<FHarvestTimerFragment>(EMassFragmentAccess::ReadOnly);
 	EntityQuery.AddTagRequirement<FFarmJustBecameReadyToHarvestTag>(EMassFragmentPresence::None);
@@ -141,7 +141,7 @@ void UFarmHarvestTimerExpired::Execute(FMassEntityManager& EntityManager, FMassE
 //-----------------------------------------------------------------------------
 // UFarmHarvestTimerSetIcon
 //-----------------------------------------------------------------------------
-void UFarmHarvestTimerSetIcon::ConfigureQueries()
+void UFarmHarvestTimerSetIcon::ConfigureQueries(const TSharedRef<FMassEntityManager>& EntityManager)
 {
 	EntityQuery.AddRequirement<FFarmGridCellData>(EMassFragmentAccess::ReadOnly);
 	EntityQuery.AddRequirement<FFarmVisualFragment>(EMassFragmentAccess::ReadWrite);
diff --git a/Engine/Source/Developer/MassEntityTestSuite/Private/MassExecutionTest.cpp b/Engine/Source/Developer/MassEntityTestSuite/Private/MassExecutionTest.cpp
--- a/Engine/Source/Developer/MassEntityTestSuite/Private/MassExecutionTest.cpp
+++ b/Engine/Source/Developer/MassEntityTestSuite/Private/MassExecutionTest.cpp
@@ -105,7 +105,7 @@ struct FExecution_SingleValidProcessor : FExecutionTestBase
 		CA_ASSUME(EntityManager);
 		const float DeltaSeconds = 0.f;
 		FMassProcessingContext ProcessingContext(EntityManager, DeltaSeconds);
-		UMassTestProcessorBase* Processor = NewObject<UMassTestProcessorBase>();
+		UMassTestProcessorBase* Processor = NewTestProcessor<UMassTestProcessorBase>(EntityManager);
 		check(Processor);
 		// need to set up some requirements to make EntityQuery valid
 		Processor->EntityQuery.AddRequirement<FTestFragment_Float>(EMassFragmentAccess::ReadOnly);
@@ -148,7 +148,7 @@ struct FExecution_Sparse : FEntityTestBase
 		CA_ASSUME(EntityManager);
 		const float DeltaSeconds = 0.f;
 		FMassProcessingContext ProcessingContext(*EntityManager, DeltaSeconds);
-		UMassTestProcessorBase* Processor = NewObject<UMassTestProcessorBase>();
+		UMassTestProcessorBase* Processor = NewTestProcessor<UMassTestProcessorBase>(EntityManager);
 		check(Processor);
 		// need to set up some requirements to make EntityQuery valid
 		Processor->EntityQuery.AddRequirement<FTestFragment_Float>(EMassFragmentAccess::ReadOnly);
diff --git a/Engine/Source/Developer/MassEntityTestSuite/Private/MassMultiThreadTest.cpp b/Engine/Source/Developer/MassEntityTestSuite/Private/MassMultiThreadTest.cpp
--- a/Engine/Source/Developer/MassEntityTestSuite/Private/MassMultiThreadTest.cpp
+++ b/Engine/Source/Developer/MassEntityTestSuite/Private/MassMultiThreadTest.cpp
@@ -72,7 +72,7 @@ struct FMTTrivial : FMTTestBase
 
 		Processors.Reset();
 		{
-			UMassTestProcessorBase* Proc = Processors.Add_GetRef(NewObject<UMassTestProcessor_A>());
+			UMassTestProcessorBase* Proc = Processors.Add_GetRef(NewTestProcessor<UMassTestProcessor_A>(EntityManager));
 
 			Proc->EntityQuery.AddRequirement<FTestFragment_Int>(EMassFragmentAccess::ReadOnly);
 			Proc->ForEachEntityChunkExecutionFunction = [this](FMassExecutionContext& Context)
@@ -109,7 +109,7 @@ struct FMTBasic : FMTTestBase
 
 		Processors.Reset();
 		{
-			UMassTestProcessorBase* Proc = Processors.Add_GetRef(NewObject<UMassTestProcessor_C>());
+			UMassTestProcessorBase* Proc = Processors.Add_GetRef(NewTestProcessor<UMassTestProcessor_C>(EntityManager));
 			Proc->GetMutableExecutionOrder().ExecuteAfter.Add(GetProcessorName<UMassTestProcessor_B>());
 			Proc->EntityQuery.AddRequirement<FTestFragment_Float>(EMassFragmentAccess::ReadOnly);
 			Proc->EntityQuery.AddRequirement<FTestFragment_Int>(EMassFragmentAccess::ReadWrite);
@@ -124,7 +124,7 @@ struct FMTBasic : FMTTestBase
 			};
 		} 
 		{
-			UMassTestProcessorBase* Proc = Processors.Add_GetRef(NewObject<UMassTestProcessor_B>());
+			UMassTestProcessorBase* Proc = Processors.Add_GetRef(NewTestProcessor<UMassTestProcessor_B>(EntityManager));
 			Proc->GetMutableExecutionOrder().ExecuteAfter.Add(GetProcessorName<UMassTestProcessor_A>());
 			Proc->EntityQuery.AddRequirement<FTestFragment_Int>(EMassFragmentAccess::ReadOnly);
 			Proc->EntityQuery.AddRequirement<FTestFragment_Float>(EMassFragmentAccess::ReadWrite);
@@ -139,7 +139,7 @@ struct FMTBasic : FMTTestBase
 			};
 		} 
 		{
-			UMassTestProcessorBase* Proc = Processors.Add_GetRef(NewObject<UMassTestProcessor_A>());
+			UMassTestProcessorBase* Proc = Processors.Add_GetRef(NewTestProcessor<UMassTestProcessor_A>(EntityManager));
 			Proc->EntityQuery.AddRequirement<FTestFragment_Int>(EMassFragmentAccess::ReadWrite);
 			Proc->ForEachEntityChunkExecutionFunction = [this](FMassExecutionContext& Context)
 			{
diff --git a/Engine/Source/Developer/MassEntityTestSuite/Private/MassObserverTest.cpp b/Engine/Source/Developer/MassEntityTestSuite/Private/MassObserverTest.cpp
--- a/Engine/Source/Developer/MassEntityTestSuite/Private/MassObserverTest.cpp
+++ b/Engine/Source/Developer/MassEntityTestSuite/Private/MassObserverTest.cpp
@@ -41,7 +41,7 @@ struct FTagBaseOperation : FEntityTestBase
 	{
 		if (FEntityTestBase::SetUp())
 		{
-			ObserverProcessor = NewObject<UMassTestProcessorBase>();
+			ObserverProcessor = NewTestProcessor<UMassTestProcessorBase>(EntityManager);
 			ObserverProcessor->EntityQuery.AddRequirement<FTestFragment_Int>(EMassFragmentAccess::ReadOnly);
 			ObserverProcessor->EntityQuery.AddTagRequirement<FTagStruct>(EMassFragmentPresence::All);
 			ObserverProcessor->ForEachEntityChunkExecutionFunction = [bCommandsFlushedPtr = &bCommandsFlushed, AffectedEntitiesPtr = &AffectedEntities](FMassExecutionContext& Context)
@@ -296,7 +296,7 @@ struct FFragmentTestBase : FEntityTestBase
 	{
 		if (FEntityTestBase::SetUp())
 		{
-			ObserverProcessor = NewObject<UMassTestProcessorBase>();
+			ObserverProcessor = NewTestProcessor<UMassTestProcessorBase>(EntityManager);
 			ObserverProcessor->EntityQuery.AddRequirement(FFragmentStruct::StaticStruct(), EMassFragmentAccess::ReadOnly);
 			ObserverProcessor->ForEachEntityChunkExecutionFunction = [bCommandsFlushedPtr = &bCommandsFlushed, AffectedEntitiesPtr = &AffectedEntities](FMassExecutionContext& Context)
 				{
diff --git a/Engine/Source/Developer/MassEntityTestSuite/Private/MassParallelQueryTest.cpp b/Engine/Source/Developer/MassEntityTestSuite/Private/MassParallelQueryTest.cpp
--- a/Engine/Source/Developer/MassEntityTestSuite/Private/MassParallelQueryTest.cpp
+++ b/Engine/Source/Developer/MassEntityTestSuite/Private/MassParallelQueryTest.cpp
@@ -27,7 +27,7 @@ struct FQueryTest_ParallelForBasic : FEntityTestBase
 		TArray<FMassEntityHandle> CreatedEntities;
 		EntityManager->BatchCreateEntities(FloatsArchetype, NumEntitiesToCreate, CreatedEntities);
 
-		FMassEntityQuery Query;
+		FMassEntityQuery Query(EntityManager);
 		Query.AddRequirement<FTestFragment_Float>(EMassFragmentAccess::ReadWrite);
 
 		std::atomic<int32> EntitiesProcessed = 0;
@@ -80,7 +80,7 @@ struct FQueryTest_ParallelForCollection : FEntityTestBase
 		const FMassArchetypeEntityCollection EntityCollection(FloatsArchetype, EntitiesToProcess, FMassArchetypeEntityCollection::NoDuplicates);
 		const int32 NumJobs = EntityCollection.GetRanges().Num();
 
-		FMassEntityQuery Query;
+		FMassEntityQuery Query(EntityManager);
 		Query.AddRequirement<FTestFragment_Float>(EMassFragmentAccess::ReadWrite);
 
 		std::atomic<int32> EntitiesProcessed = 0;
@@ -128,7 +128,7 @@ struct FQueryTest_ParallelCommands : FEntityTestBase
 		TArray<FMassEntityHandle> CreatedEntities;
 		EntityManager->BatchCreateEntities(FloatsArchetype, NumEntitiesToCreate, CreatedEntities);
 
-		FMassEntityQuery Query;
+		FMassEntityQuery Query(EntityManager);
 		Query.AddRequirement<FTestFragment_Float>(EMassFragmentAccess::ReadWrite);
 		Query.SetParallelCommandBufferEnabled(true);
 
diff --git a/Engine/Source/Developer/MassEntityTestSuite/Private/MassProcessorTest.cpp b/Engine/Source/Developer/MassEntityTestSuite/Private/MassProcessorTest.cpp
--- a/Engine/Source/Developer/MassEntityTestSuite/Private/MassProcessorTest.cpp
+++ b/Engine/Source/Developer/MassEntityTestSuite/Private/MassProcessorTest.cpp
@@ -19,6 +19,7 @@ int32 SimpleProcessorRun(FMassEntityManager& EntityManager)
 {
 	int32 EntityProcessedCount = 0;
 	TProcessor* Processor = NewObject<TProcessor>();
+	Processor->CallInitialize(GetTransientPackage(), EntityManager.AsShared());
 	Processor->ForEachEntityChunkExecutionFunction = [&EntityProcessedCount](FMassExecutionContext& Context) 
 	{
 		EntityProcessedCount += Context.GetNumEntities();
diff --git a/Engine/Source/Developer/MassEntityTestSuite/Private/MassQueryTest.cpp b/Engine/Source/Developer/MassEntityTestSuite/Private/MassQueryTest.cpp
--- a/Engine/Source/Developer/MassEntityTestSuite/Private/MassQueryTest.cpp
+++ b/Engine/Source/Developer/MassEntityTestSuite/Private/MassQueryTest.cpp
@@ -22,7 +22,7 @@ struct FQueryTest_ProcessorRequirements : FEntityTestBase
 	{
 		CA_ASSUME(EntityManager);
 
-		UMassTestProcessor_Floats* Processor = NewObject<UMassTestProcessor_Floats>();
+		UMassTestProcessor_Floats* Processor = NewTestProcessor<UMassTestProcessor_Floats>(EntityManager);
 		TConstArrayView<FMassFragmentRequirementDescription> Requirements = Processor->EntityQuery.GetFragmentRequirements();
 		
 		AITEST_TRUE("Query should have extracted some requirements from the given Processor", Requirements.Num() > 0);
@@ -41,7 +41,7 @@ struct FQueryTest_ExplicitRequirements : FEntityTestBase
 	{
 		CA_ASSUME(EntityManager);
 				
-		FMassEntityQuery Query({ FTestFragment_Float::StaticStruct()});
+		FMassEntityQuery Query(EntityManager.ToSharedRef(), { FTestFragment_Float::StaticStruct()});
 		TConstArrayView<FMassFragmentRequirementDescription> Requirements = Query.GetFragmentRequirements();
 
 		AITEST_TRUE("Query should have extracted some requirements from the given Processor", Requirements.Num() > 0);
@@ -64,7 +64,7 @@ struct FQueryTest_FragmentViewBinding : FEntityTestBase
 		FTestFragment_Float& TestedFragment = EntityManager->GetFragmentDataChecked<FTestFragment_Float>(Entity);
 		AITEST_TRUE("Initial value of the fragment should match expectations", TestedFragment.Value == 0.f);
 
-		UMassTestProcessor_Floats* Processor = NewObject<UMassTestProcessor_Floats>();
+		UMassTestProcessor_Floats* Processor = NewTestProcessor<UMassTestProcessor_Floats>(EntityManager);
 		Processor->ForEachEntityChunkExecutionFunction = [](FMassExecutionContext& Context) 
 		{
 			TArrayView<FTestFragment_Float> Floats = Context.GetMutableFragmentView<FTestFragment_Float>();
@@ -97,7 +97,7 @@ struct FQueryTest_ExecuteSingleArchetype : FEntityTestBase
 		int TotalProcessed = 0;
 
 		FMassExecutionContext ExecContext(*EntityManager.Get());
-		FMassEntityQuery Query({ FTestFragment_Float::StaticStruct() });
+		FMassEntityQuery Query(EntityManager.ToSharedRef(), { FTestFragment_Float::StaticStruct() });
 		Query.ForEachEntityChunk(*EntityManager, ExecContext, [&TotalProcessed](FMassExecutionContext& Context)
 			{
 				TotalProcessed += Context.GetNumEntities();
@@ -141,7 +141,7 @@ struct FQueryTest_ExecuteMultipleArchetypes : FEntityTestBase
 
 		int TotalProcessed = 0;
 		FMassExecutionContext ExecContext(*EntityManager.Get());
-		FMassEntityQuery Query({ FTestFragment_Float::StaticStruct() });
+		FMassEntityQuery Query(EntityManager.ToSharedRef(), { FTestFragment_Float::StaticStruct() });
 		Query.ForEachEntityChunk(*EntityManager, ExecContext, [&TotalProcessed](FMassExecutionContext& Context)
 			{
 				TotalProcessed += Context.GetNumEntities();
@@ -197,7 +197,7 @@ struct FQueryTest_ExecuteSparse : FEntityTestBase
 		int TotalProcessed = 0;
 
 		FMassExecutionContext ExecContext(*EntityManager.Get());
-		FMassEntityQuery TestQuery;
+		FMassEntityQuery TestQuery(EntityManager);
 		TestQuery.AddRequirement<FTestFragment_Float>(EMassFragmentAccess::ReadWrite);
 		TestQuery.ForEachEntityChunk(FMassArchetypeEntityCollection(FloatsArchetype, EntitiesToProcess, FMassArchetypeEntityCollection::NoDuplicates)
 								, *EntityManager, ExecContext, [&TotalProcessed](FMassExecutionContext& Context)
@@ -240,7 +240,7 @@ struct FQueryTest_TagPresent : FEntityTestBase
 		TArray<const UScriptStruct*> Fragments = {FTestFragment_Float::StaticStruct(), FTestFragment_Tag::StaticStruct()};
 		const FMassArchetypeHandle FloatsTagArchetype = EntityManager->CreateArchetype(Fragments);
 
-		FMassEntityQuery Query;
+		FMassEntityQuery Query(EntityManager);
 		Query.AddRequirement<FTestFragment_Float>(EMassFragmentAccess::ReadWrite);
 		Query.AddTagRequirement<FTestFragment_Tag>(EMassFragmentPresence::All);
 		Query.CacheArchetypes(*EntityManager);
@@ -263,7 +263,7 @@ struct FQueryTest_TagAbsent : FEntityTestBase
 		TArray<const UScriptStruct*> Fragments = { FTestFragment_Float::StaticStruct(), FTestFragment_Tag::StaticStruct() };
 		const FMassArchetypeHandle FloatsTagArchetype = EntityManager->CreateArchetype(Fragments);
 
-		FMassEntityQuery Query;
+		FMassEntityQuery Query(EntityManager);
 		Query.AddRequirement<FTestFragment_Float>(EMassFragmentAccess::ReadWrite);
 		Query.AddTagRequirement<FTestFragment_Tag>(EMassFragmentPresence::None);
 		Query.CacheArchetypes(*EntityManager);
@@ -285,7 +285,7 @@ struct FQueryTest_FragmentPresent : FEntityTestBase
 	{
 		CA_ASSUME(EntityManager);
 
-		FMassEntityQuery Query;
+		FMassEntityQuery Query(EntityManager);
 		// using EMassFragmentAccess::None to indicate we're interested only in the archetype having the fragment, no binding is required
 		Query.AddRequirement<FTestFragment_Int>(EMassFragmentAccess::None, EMassFragmentPresence::Any);
 		Query.CacheArchetypes(*EntityManager);
@@ -320,7 +320,7 @@ struct FQueryTest_OnlyAbsentFragments : FEntityTestBase
 	{
 		CA_ASSUME(EntityManager);
 
-		FMassEntityQuery Query;
+		FMassEntityQuery Query(EntityManager);
 		AITEST_FALSE("The empty query is not valid", Query.CheckValidity());
 
 		Query.AddRequirement<FTestFragment_Int>(EMassFragmentAccess::ReadOnly, EMassFragmentPresence::None);
@@ -351,7 +351,7 @@ struct FQueryTest_AbsentAndPresentFragments : FEntityTestBase
 	{
 		CA_ASSUME(EntityManager);
 
-		FMassEntityQuery Query;
+		FMassEntityQuery Query(EntityManager);
 		Query.AddRequirement<FTestFragment_Int>(EMassFragmentAccess::None, EMassFragmentPresence::None);
 		Query.AddRequirement<FTestFragment_Float>(EMassFragmentAccess::ReadOnly, EMassFragmentPresence::All);
 
@@ -372,7 +372,7 @@ struct FQueryTest_SingleOptionalFragment : FEntityTestBase
 	{
 		CA_ASSUME(EntityManager);
 
-		FMassEntityQuery Query;
+		FMassEntityQuery Query(EntityManager);
 		Query.AddRequirement<FTestFragment_Int>(EMassFragmentAccess::ReadWrite, EMassFragmentPresence::Optional);
 		Query.CacheArchetypes(*EntityManager);
 
@@ -392,7 +392,7 @@ struct FQueryTest_MultipleOptionalFragment : FEntityTestBase
 	{
 		CA_ASSUME(EntityManager);
 
-		FMassEntityQuery Query;
+		FMassEntityQuery Query(EntityManager);
 		Query.AddRequirement<FTestFragment_Int>(EMassFragmentAccess::ReadWrite, EMassFragmentPresence::Optional);
 		Query.AddRequirement<FTestFragment_Float>(EMassFragmentAccess::ReadWrite, EMassFragmentPresence::Optional);
 		Query.CacheArchetypes(*EntityManager);
@@ -420,7 +420,7 @@ struct FQueryTest_UsingOptionalFragment : FEntityTestBase
 		int TotalProcessed = 0;
 		int EmptyIntsViewCount = 0;
 
-		FMassEntityQuery Query;
+		FMassEntityQuery Query(EntityManager);
 		Query.AddRequirement<FTestFragment_Int>(EMassFragmentAccess::ReadWrite, EMassFragmentPresence::Optional);
 		Query.AddRequirement<FTestFragment_Float>(EMassFragmentAccess::ReadWrite, EMassFragmentPresence::All);
 		FMassExecutionContext ExecContext(*EntityManager.Get());
@@ -465,7 +465,7 @@ struct FQueryTest_AnyFragment : FEntityTestBase
 		const FMassArchetypeHandle BoolArchetype = EntityManager->CreateArchetype({ FTestFragment_Bool::StaticStruct() });
 		const FMassArchetypeHandle BoolFloatArchetype = EntityManager->CreateArchetype({ FTestFragment_Bool::StaticStruct(), FTestFragment_Float::StaticStruct() });
 
-		FMassEntityQuery Query;
+		FMassEntityQuery Query(EntityManager);
 		Query.AddRequirement<FTestFragment_Int>(EMassFragmentAccess::ReadWrite, EMassFragmentPresence::Any);
 		Query.AddRequirement<FTestFragment_Bool>(EMassFragmentAccess::ReadWrite, EMassFragmentPresence::Any);
 		// this query should match: 
@@ -507,7 +507,7 @@ struct FQueryTest_AnyTag : FEntityTestBase
 		const FMassArchetypeHandle BDArchetype = EntityManager->CreateArchetype({ FTestFragment_Int::StaticStruct(), FTestTag_B::StaticStruct(), FTestTag_D::StaticStruct() });
 		const FMassArchetypeHandle FloatACArchetype = EntityManager->CreateArchetype({ FTestFragment_Float::StaticStruct(), FTestTag_A::StaticStruct(), FTestTag_C::StaticStruct() });
 
-		FMassEntityQuery Query;
+		FMassEntityQuery Query(EntityManager);
 		// at least one fragment requirement needs to be present for the query to be valid
 		Query.AddRequirement<FTestFragment_Int>(EMassFragmentAccess::ReadOnly); 
 		Query.AddTagRequirement<FTestTag_A>(EMassFragmentPresence::Any);
@@ -533,7 +533,7 @@ struct FQueryTest_AutoRecache : FEntityTestBase
 	{
 		CA_ASSUME(EntityManager);
 
-		FMassEntityQuery Query;
+		FMassEntityQuery Query(EntityManager);
 		// at least one fragment requirement needs to be present for the query to be valid
 		Query.AddRequirement<FTestFragment_Int>(EMassFragmentAccess::ReadOnly);
 		
@@ -582,7 +582,7 @@ struct FQueryTest_AllOptional : FEntityTestBase
 	{
 		CA_ASSUME(EntityManager);
 
-		FMassEntityQuery Query;
+		FMassEntityQuery Query(EntityManager);
 		Query.AddRequirement<FTestFragment_Float>(EMassFragmentAccess::None, EMassFragmentPresence::Optional);
 		Query.AddTagRequirement<FTestTag_A>(EMassFragmentPresence::Optional);
 		Query.AddChunkRequirement<FTestChunkFragment_Int>(EMassFragmentAccess::None, EMassFragmentPresence::Optional);
@@ -667,7 +667,7 @@ struct FQueryTest_JustATag : FEntityTestBase
 	{
 		CA_ASSUME(EntityManager);
 
-		FMassEntityQuery Query;
+		FMassEntityQuery Query(EntityManager);
 		Query.AddTagRequirement<FTestTag_A>(EMassFragmentPresence::All);
 		Query.CacheArchetypes(*EntityManager);
 
@@ -720,7 +720,7 @@ struct FQueryTest_JustAChunkFragment : FEntityTestBase
 
 		FMassArchetypeHandle TargetArchetype;
 
-		FMassEntityQuery Query;
+		FMassEntityQuery Query(EntityManager);
 		Query.AddChunkRequirement<FTestChunkFragment_Int>(EMassFragmentAccess::ReadOnly, EMassFragmentPresence::All);
 		Query.CacheArchetypes(*EntityManager);
 
@@ -770,7 +770,7 @@ struct FQueryTest_JustASharedFragment : FEntityTestBase
 
 		FMassArchetypeHandle TargetArchetype;
 
-		FMassEntityQuery Query;
+		FMassEntityQuery Query(EntityManager);
 		Query.AddSharedRequirement<FTestSharedFragment_Int>(EMassFragmentAccess::ReadOnly, EMassFragmentPresence::All);
 		Query.CacheArchetypes(*EntityManager);
 
@@ -830,7 +830,7 @@ struct FQueryTest_JustAConstSharedFragment : FEntityTestBase
 
 		FMassArchetypeHandle TargetArchetype;
 
-		FMassEntityQuery Query;
+		FMassEntityQuery Query(EntityManager);
 		Query.AddConstSharedRequirement<FTestConstSharedFragment_Int>(EMassFragmentPresence::All);
 		Query.CacheArchetypes(*EntityManager);
 
diff --git a/Engine/Source/Developer/MassEntityTestSuite/Private/MassSystemRequirementTest.cpp b/Engine/Source/Developer/MassEntityTestSuite/Private/MassSystemRequirementTest.cpp
--- a/Engine/Source/Developer/MassEntityTestSuite/Private/MassSystemRequirementTest.cpp
+++ b/Engine/Source/Developer/MassEntityTestSuite/Private/MassSystemRequirementTest.cpp
@@ -32,6 +32,7 @@ struct FSystemRequirementTestBase : FEntityTestBase
 		if (Super::SetUp())
 		{
 			EntityManager->CreateEntity(FloatsArchetype);
+			EntityQuery.Initialize(EntityManager.ToSharedRef());
 			EntityQuery.AddRequirement<FTestFragment_Float>(EMassFragmentAccess::ReadOnly);
 			return true;
 		}
diff --git a/Engine/Source/Developer/MassEntityTestSuite/Private/MassTestTypes.cpp b/Engine/Source/Developer/MassEntityTestSuite/Private/MassTestTypes.cpp
--- a/Engine/Source/Developer/MassEntityTestSuite/Private/MassTestTypes.cpp
+++ b/Engine/Source/Developer/MassEntityTestSuite/Private/MassTestTypes.cpp
@@ -65,17 +65,17 @@ void UMassTestProcessorBase::Execute(FMassEntityManager& EntityManager, FMassExe
 	ExecutionFunction(EntityManager, Context);
 }
 
-UMassTestProcessor_Floats::UMassTestProcessor_Floats()
+void UMassTestProcessor_Floats::ConfigureQueries(const TSharedRef<FMassEntityManager>& EntityManager)
 {
 	EntityQuery.AddRequirement<FTestFragment_Float>(EMassFragmentAccess::ReadWrite);
 }
 
-UMassTestProcessor_Ints::UMassTestProcessor_Ints()
+void UMassTestProcessor_Ints::ConfigureQueries(const TSharedRef<FMassEntityManager>& EntityManager)
 {
 	EntityQuery.AddRequirement<FTestFragment_Int>(EMassFragmentAccess::ReadWrite);
 }
 
-UMassTestProcessor_FloatsInts::UMassTestProcessor_FloatsInts()
+void UMassTestProcessor_FloatsInts::ConfigureQueries(const TSharedRef<FMassEntityManager>& EntityManager)
 {
 	EntityQuery.AddRequirement<FTestFragment_Float>(EMassFragmentAccess::ReadWrite);
 	EntityQuery.AddRequirement<FTestFragment_Int>(EMassFragmentAccess::ReadWrite);
diff --git a/Engine/Source/Developer/MassEntityTestSuite/Public/MassEntityTestFarmPlot.h b/Engine/Source/Developer/MassEntityTestSuite/Public/MassEntityTestFarmPlot.h
--- a/Engine/Source/Developer/MassEntityTestSuite/Public/MassEntityTestFarmPlot.h
+++ b/Engine/Source/Developer/MassEntityTestSuite/Public/MassEntityTestFarmPlot.h
@@ -120,7 +120,7 @@ class UFarmWaterProcessor : public UFarmProcessorBase
 	GENERATED_BODY()
 
 public:
-	virtual void ConfigureQueries() override;
+	virtual void ConfigureQueries(const TSharedRef<FMassEntityManager>& EntityManager) override;
 	virtual void Execute(FMassEntityManager& EntityManager, FMassExecutionContext& Context) override;
 };
 
@@ -131,7 +131,7 @@ class UFarmHarvestTimerSystem_Flowers : public UFarmProcessorBase
 {
 	GENERATED_BODY()
 
-	virtual void ConfigureQueries() override;
+	virtual void ConfigureQueries(const TSharedRef<FMassEntityManager>& EntityManager) override;
 	virtual void Execute(FMassEntityManager& EntityManager, FMassExecutionContext& Context) override;
 };
 
@@ -142,7 +142,7 @@ class UFarmHarvestTimerSystem_Crops : public UFarmProcessorBase
 {
 	GENERATED_BODY()
 
-	virtual void ConfigureQueries() override;
+	virtual void ConfigureQueries(const TSharedRef<FMassEntityManager>& EntityManager) override;
 	virtual void Execute(FMassEntityManager& EntityManager, FMassExecutionContext& Context) override;
 };
 
@@ -153,7 +153,7 @@ class UFarmHarvestTimerExpired : public UFarmProcessorBase
 {
 	GENERATED_BODY()
 
-	virtual void ConfigureQueries() override;
+	virtual void ConfigureQueries(const TSharedRef<FMassEntityManager>& EntityManager) override;
 	virtual void Execute(FMassEntityManager& EntityManager, FMassExecutionContext& Context) override;
 };
 
@@ -174,7 +174,7 @@ public:
 	float HarvestIconScale;
 
 protected:
-	virtual void ConfigureQueries() override;
+	virtual void ConfigureQueries(const TSharedRef<FMassEntityManager>& EntityManager) override;
 	virtual void Execute(FMassEntityManager& EntityManager, FMassExecutionContext& Context) override;
 };
 
diff --git a/Engine/Source/Developer/MassEntityTestSuite/Public/MassEntityTestTypes.h b/Engine/Source/Developer/MassEntityTestSuite/Public/MassEntityTestTypes.h
--- a/Engine/Source/Developer/MassEntityTestSuite/Public/MassEntityTestTypes.h
+++ b/Engine/Source/Developer/MassEntityTestSuite/Public/MassEntityTestTypes.h
@@ -205,7 +205,7 @@ public:
 	 * since the EntityQuery configuration will take place after a processor instance gets created
 	 * (so after PostInitProperties has already been called). Use EntityQuery directly instead. 
 	 */
-	virtual void ConfigureQueries() override {}
+	virtual void ConfigureQueries(const TSharedRef<FMassEntityManager>&) override {}
 
 	using FExecutionFunction = TFunction<void(FMassEntityManager& EntityManager, FMassExecutionContext& Context)>;
 	FExecutionFunction ExecutionFunction;
@@ -265,7 +265,7 @@ class UMassTestProcessor_Floats : public UMassTestProcessorBase
 	GENERATED_BODY()
 public:
 	TArrayView<FTestFragment_Float> Floats;
-	UMassTestProcessor_Floats();
+	virtual void ConfigureQueries(const TSharedRef<FMassEntityManager>& EntityManager) override;
 };
 
 UCLASS()
@@ -274,7 +274,7 @@ class UMassTestProcessor_Ints : public UMassTestProcessorBase
 	GENERATED_BODY()
 public:
 	TArrayView<FTestFragment_Int> Ints;
-	UMassTestProcessor_Ints();
+	virtual void ConfigureQueries(const TSharedRef<FMassEntityManager>& EntityManager) override;
 };
 
 UCLASS()
@@ -284,9 +284,19 @@ class UMassTestProcessor_FloatsInts : public UMassTestProcessorBase
 public:
 	TArrayView<FTestFragment_Float> Floats;
 	TArrayView<FTestFragment_Int> Ints;
-	UMassTestProcessor_FloatsInts();
+	virtual void ConfigureQueries(const TSharedRef<FMassEntityManager>& EntityManager) override;
 };
 
+template<typename T>
+T* NewTestProcessor(TSharedPtr<FMassEntityManager> EntityManager)
+{
+	check(EntityManager);
+	T* NewProcessor = NewObject<T>();
+	CA_ASSUME(NewProcessor);
+	NewProcessor->CallInitialize(GetTransientPackage(), EntityManager.ToSharedRef());
+	return NewProcessor;
+}
+
 UCLASS()
 class MASSENTITYTESTSUITE_API UMassTestStaticCounterProcessor : public UMassProcessor
 {
@@ -298,8 +308,8 @@ public:
 	{
 		++StaticCounter;
 	}
-	virtual void ConfigureQueries() override {}
-	virtual bool ShouldAllowQueryBasedPruning(const bool bRuntimeMode = true) const { return false; }
+	virtual void ConfigureQueries(const TSharedRef<FMassEntityManager>&) override {}
+	virtual bool ShouldAllowQueryBasedPruning(const bool bRuntimeMode = true) const override { return false; }
 
 	static int StaticCounter;
 };
diff --git a/Engine/Source/Runtime/MassEntity/Private/MassEntityQuery.cpp b/Engine/Source/Runtime/MassEntity/Private/MassEntityQuery.cpp
--- a/Engine/Source/Runtime/MassEntity/Private/MassEntityQuery.cpp
+++ b/Engine/Source/Runtime/MassEntity/Private/MassEntityQuery.cpp
@@ -52,12 +52,18 @@ FMassEntityQuery::FScopedSubsystemRequirementsRestore::~FScopedSubsystemRequirem
 //-----------------------------------------------------------------------------
 // FMassEntityQuery
 //-----------------------------------------------------------------------------
-FMassEntityQuery::FMassEntityQuery()
+FMassEntityQuery::FMassEntityQuery(UMassProcessor& Owner)
 {
+	RegisterWithProcessor(Owner);
 }
 
-FMassEntityQuery::FMassEntityQuery(std::initializer_list<UScriptStruct*> InitList)
-	: FMassEntityQuery()
+FMassEntityQuery::FMassEntityQuery(const TSharedPtr<FMassEntityManager>& EntityManager)
+	: FMassFragmentRequirements(EntityManager)
+{
+}
+
+FMassEntityQuery::FMassEntityQuery(const TSharedRef<FMassEntityManager>& EntityManager, std::initializer_list<UScriptStruct*> InitList)
+	: FMassFragmentRequirements(EntityManager)
 {
 	for (const UScriptStruct* FragmentType : InitList)
 	{
@@ -65,8 +71,8 @@ FMassEntityQuery::FMassEntityQuery(std::initializer_list<UScriptStruct*> InitLis
 	}
 }
 
-FMassEntityQuery::FMassEntityQuery(TConstArrayView<const UScriptStruct*> InitList)
-	: FMassEntityQuery()
+FMassEntityQuery::FMassEntityQuery(const TSharedRef<FMassEntityManager>& EntityManager, TConstArrayView<const UScriptStruct*> InitList)
+	: FMassFragmentRequirements(EntityManager)
 {
 	for (const UScriptStruct* FragmentType : InitList)
 	{
@@ -74,11 +80,6 @@ FMassEntityQuery::FMassEntityQuery(TConstArrayView<const UScriptStruct*> InitLis
 	}
 }
 
-FMassEntityQuery::FMassEntityQuery(UMassProcessor& Owner)
-{
-	RegisterWithProcessor(Owner);
-}
-
 void FMassEntityQuery::RegisterWithProcessor(UMassProcessor& Owner)
 {
 	ExpectedContextType = EMassExecutionContextType::Processor;
@@ -505,3 +506,16 @@ void FMassEntityQuery::ExportRequirements(FMassExecutionRequirements& OutRequire
 	FMassSubsystemRequirements::ExportRequirements(OutRequirements);
 	FMassFragmentRequirements::ExportRequirements(OutRequirements);
 }
+
+//-----------------------------------------------------------------------------
+// DEPRECATED
+//-----------------------------------------------------------------------------
+FMassEntityQuery::FMassEntityQuery(std::initializer_list<UScriptStruct*>)
+	: FMassEntityQuery()
+{
+}
+
+FMassEntityQuery::FMassEntityQuery(TConstArrayView<const UScriptStruct*>)
+	: FMassEntityQuery()
+{
+}
\ No newline at end of file
diff --git a/Engine/Source/Runtime/MassEntity/Private/MassProcessor.cpp b/Engine/Source/Runtime/MassEntity/Private/MassProcessor.cpp
--- a/Engine/Source/Runtime/MassEntity/Private/MassProcessor.cpp
+++ b/Engine/Source/Runtime/MassEntity/Private/MassProcessor.cpp
@@ -127,7 +127,16 @@ void UMassProcessor::CallInitialize(const gsl::not_null<UObject*> Owner, const T
 	if (ensure(HasAnyFlags(RF_ClassDefaultObject) == false 
 		&& GetClass()->HasAnyClassFlags(CLASS_Abstract) == false))
 	{
-		ConfigureQueries();
+		for (FMassEntityQuery* Query : OwnedQueries)
+		{
+			// we should never get nulls here since OwnedQueries is private and the only way to
+			// add queries to it is to go through RegisterQuery, which in turn ensures the
+			// input query is a member variable of the processor.
+			checkfSlow(Query, TEXT("We never expect nulls in OwnedQueries - those pointers are supposed to point at member variable."));
+			Query->Initialize(EntityManager);
+		}
+
+		ConfigureQueries(EntityManager);
 
 		bool bNeedsGameThread = ProcessorRequirements.DoesRequireGameThreadExecution();
 		for (const FMassEntityQuery* QueryPtr : OwnedQueries)
@@ -336,7 +345,7 @@ void UMassCompositeProcessor::SetChildProcessors(TArray<TObjectPtr<UMassProcesso
 	ChildPipeline.SetProcessors(MoveTemp(InProcessors));
 }
 
-void UMassCompositeProcessor::ConfigureQueries()
+void UMassCompositeProcessor::ConfigureQueries(const TSharedRef<FMassEntityManager>& EntityManager)
 {
 	// nothing to do here since ConfigureQueries will get independently called for all the processors during their creation
 }
diff --git a/Engine/Source/Runtime/MassEntity/Private/MassRequirements.cpp b/Engine/Source/Runtime/MassEntity/Private/MassRequirements.cpp
--- a/Engine/Source/Runtime/MassEntity/Private/MassRequirements.cpp
+++ b/Engine/Source/Runtime/MassEntity/Private/MassRequirements.cpp
@@ -70,20 +70,30 @@ void FMassSubsystemRequirements::Reset()
 //-----------------------------------------------------------------------------
 // FMassFragmentRequirements
 //-----------------------------------------------------------------------------
-FMassFragmentRequirements::FMassFragmentRequirements(std::initializer_list<UScriptStruct*> InitList)
+FMassFragmentRequirements::FMassFragmentRequirements(const TSharedPtr<FMassEntityManager>& EntityManager)
 {
-	for (const UScriptStruct* FragmentType : InitList)
+	if (ensure(EntityManager))
 	{
-		AddRequirement(FragmentType, EMassFragmentAccess::ReadWrite, EMassFragmentPresence::All);
+		Initialize(EntityManager.ToSharedRef());
 	}
 }
 
-FMassFragmentRequirements::FMassFragmentRequirements(TConstArrayView<const UScriptStruct*> InitList)
+FMassFragmentRequirements::FMassFragmentRequirements(const TSharedRef<FMassEntityManager>& EntityManager)
 {
-	for (const UScriptStruct* FragmentType : InitList)
+	Initialize(EntityManager);
+}
+
+void FMassFragmentRequirements::Initialize(const TSharedRef<FMassEntityManager>& EntityManager)
+{
+	UE_CLOG(CachedEntityManager && (CachedEntityManager != EntityManager), LogMass, Warning
+		, TEXT("Trying to initialize FMassFragmentRequirements with another entity manager"));
+	if (bInitialized)
 	{
-		AddRequirement(FragmentType, EMassFragmentAccess::ReadWrite, EMassFragmentPresence::All);
+		return;
 	}
+
+	CachedEntityManager = EntityManager;
+	bInitialized = true;
 }
 
 FMassFragmentRequirements& FMassFragmentRequirements::ClearTagRequirements(const FMassTagBitSet& TagsToRemoveBitSet)
@@ -108,7 +118,7 @@ void FMassFragmentRequirements::SortRequirements()
 	SharedFragmentRequirements.Sort(FScriptStructSortOperator());
 }
 
-FORCEINLINE void FMassFragmentRequirements::CachePropreties() const
+FORCEINLINE void FMassFragmentRequirements::CacheProperties() const
 {
 	if (bPropertiesCached == false)
 	{
@@ -138,14 +148,14 @@ FORCEINLINE void FMassFragmentRequirements::CachePropreties() const
 
 bool FMassFragmentRequirements::CheckValidity() const
 {
-	CachePropreties();
+	CacheProperties();
 	// @todo we need to add more sophisticated testing somewhere to detect contradicting requirements - like having and not having a given tag.
 	return bHasPositiveRequirements || bHasNegativeRequirements || bHasOptionalRequirements;
 }
 
 bool FMassFragmentRequirements::IsEmpty() const
 {
-	CachePropreties();
+	CacheProperties();
 	// note that even though at the moment the following condition is the same as negation of current CheckValidity value
 	// that will change in the future (with additional validity checks).
 	return !bHasPositiveRequirements && !bHasNegativeRequirements && !bHasOptionalRequirements;
@@ -172,7 +182,7 @@ bool FMassFragmentRequirements::DoesArchetypeMatchRequirements(const FMassArchet
 	
 bool FMassFragmentRequirements::DoesArchetypeMatchRequirements(const FMassArchetypeCompositionDescriptor& ArchetypeComposition) const
 {
-	CachePropreties();
+	CacheProperties();
 
 	const bool bPassNegativeFilter = bHasNegativeRequirements == false
 		|| (ArchetypeComposition.Fragments.HasNone(RequiredNoneFragments)
@@ -241,4 +251,27 @@ void FMassFragmentRequirements::Reset()
 	RequiredNoneConstSharedFragments.Reset();
 
 	IncrementalChangesCount = 0;
+
+	// note that we're not resetting bInitialized nor CachedEntityManager, on purpose
+	// the point of this function is to just reset the contents while still being able
+	// to add elements to it. This "requirements" instance is now "empty" but still valid
+}
+
+//-----------------------------------------------------------------------------
+// DEPRECATED
+//-----------------------------------------------------------------------------
+FMassFragmentRequirements::FMassFragmentRequirements(std::initializer_list<UScriptStruct*> InitList)
+{
+	for (const UScriptStruct* FragmentType : InitList)
+	{
+		AddRequirement(FragmentType, EMassFragmentAccess::ReadWrite, EMassFragmentPresence::All);
+	}
+}
+
+FMassFragmentRequirements::FMassFragmentRequirements(TConstArrayView<const UScriptStruct*> InitList)
+{
+	for (const UScriptStruct* FragmentType : InitList)
+	{
+		AddRequirement(FragmentType, EMassFragmentAccess::ReadWrite, EMassFragmentPresence::All);
+	}
 }
diff --git a/Engine/Source/Runtime/MassEntity/Public/MassEntityQuery.h b/Engine/Source/Runtime/MassEntity/Public/MassEntityQuery.h
--- a/Engine/Source/Runtime/MassEntity/Public/MassEntityQuery.h
+++ b/Engine/Source/Runtime/MassEntity/Public/MassEntityQuery.h
@@ -33,6 +33,8 @@ concept TArchetypeMatchOverrideConcept = requires(const FMassArchetypeCompositio
 };
 #endif
 
+struct FMassEntityManager;
+
 /** 
  *  FMassEntityQuery is a structure that is used to trigger calculations on cached set of valid archetypes as described 
  *  by requirements. See the parent classes FMassFragmentRequirements and FMassSubsystemRequirements for setting up the 
@@ -56,7 +58,7 @@ private:
     {
 		using MatchFunction = bool (*)(const void* /*Context*/, const FMassArchetypeCompositionDescriptor&);
 		
-		MatchFunction Match;
+		MatchFunction Match = nullptr;
     	TAlignedBytes<ArchetypeMatchOverrideSize, ArchetypeMatchOverrideAlignment> Data;
     };
 
@@ -67,10 +69,11 @@ public:
 		ForceParallelExecution // force
 	};
 
-	FMassEntityQuery();
-	FMassEntityQuery(std::initializer_list<UScriptStruct*> InitList);
-	FMassEntityQuery(TConstArrayView<const UScriptStruct*> InitList);
+	FMassEntityQuery() = default;
 	FMassEntityQuery(UMassProcessor& Owner);
+	FMassEntityQuery(const TSharedPtr<FMassEntityManager>& EntityManager);
+	FMassEntityQuery(const TSharedRef<FMassEntityManager>& EntityManager, std::initializer_list<UScriptStruct*> InitList);
+	FMassEntityQuery(const TSharedRef<FMassEntityManager>& EntityManager, TConstArrayView<const UScriptStruct*> InitList);
 
 	void RegisterWithProcessor(UMassProcessor& Owner);
 
@@ -220,12 +223,22 @@ private:
 	EMassExecutionContextType ExpectedContextType = EMassExecutionContextType::Local;
 
 #if WITH_MASSENTITY_DEBUG
-	uint8 bRegistered : 1;
+	uint8 bRegistered : 1 = false;
 #endif // WITH_MASSENTITY_DEBUG
 
 #if WITH_ARCHETYPE_MATCH_OVERRIDE
 	FArchetypeMatchOverride ArchetypeMatchOverride;
 #endif
+
+public:
+	//-----------------------------------------------------------------------------
+	// DEPRECATED
+	//-----------------------------------------------------------------------------
+	UE_DEPRECATED(5.6, "This type of FMassEntityQuery is no longer supported. Use one of the other constructors instead.")
+	FMassEntityQuery(std::initializer_list<UScriptStruct*> InitList);
+
+	UE_DEPRECATED(5.6, "This type of FMassEntityQuery is no longer supported. Use one of the other constructors instead.")
+	FMassEntityQuery(TConstArrayView<const UScriptStruct*> InitList);
 };
 
 #if WITH_ARCHETYPE_MATCH_OVERRIDE
diff --git a/Engine/Source/Runtime/MassEntity/Public/MassProcessor.h b/Engine/Source/Runtime/MassEntity/Public/MassProcessor.h
--- a/Engine/Source/Runtime/MassEntity/Public/MassProcessor.h
+++ b/Engine/Source/Runtime/MassEntity/Public/MassProcessor.h
@@ -152,7 +152,7 @@ protected:
 	/** Called to initialize the processor. Override to perform custom steps. */
 	virtual void InitializeInternal(UObject& Owner, const TSharedRef<FMassEntityManager>& EntityManager);
 
-	virtual void ConfigureQueries() PURE_VIRTUAL(UMassProcessor::ConfigureQueries);
+	virtual void ConfigureQueries(const TSharedRef<FMassEntityManager>& EntityManager) PURE_VIRTUAL(UMassProcessor::ConfigureQueries);
 	virtual void PostInitProperties() override;
 	virtual void Execute(FMassEntityManager& EntityManager, FMassExecutionContext& Context) PURE_VIRTUAL(UMassProcessor::Execute);
 
@@ -235,7 +235,9 @@ private:
 	//-----------------------------------------------------------------------------
 public:
 	UE_DEPRECATED(5.6, "Initialize is deprecated. Override Initialize(UObject&, const TSharedRef<FMassEntityManager>&) instead. If you want to call the function, use CallInitialize.")
-	virtual void Initialize(UObject& Owner) final;
+	virtual void Initialize(UObject& Owner);
+	UE_DEPRECATED(5.6, "This flavor of ConfigureQueries is deprecated. Override ConfigureQueries(const TSharedRef<FMassEntityManager>&) instead.")
+	virtual void ConfigureQueries() {};
 };
 
 
@@ -299,7 +301,7 @@ public:
 	TConstArrayView<UMassProcessor*> GetChildProcessorsView() const;
 
 protected:
-	virtual void ConfigureQueries() override;
+	virtual void ConfigureQueries(const TSharedRef<FMassEntityManager>& EntityManager) override;
 	virtual void Execute(FMassEntityManager& EntityManager, FMassExecutionContext& Context) override;
 
 	/** RequestedGroupName can indicate a multi-level group name, like so: A.B.C
diff --git a/Engine/Source/Runtime/MassEntity/Public/MassRequirements.h b/Engine/Source/Runtime/MassEntity/Public/MassRequirements.h
--- a/Engine/Source/Runtime/MassEntity/Public/MassRequirements.h
+++ b/Engine/Source/Runtime/MassEntity/Public/MassRequirements.h
@@ -165,12 +165,15 @@ struct MASSENTITY_API FMassFragmentRequirements
 	friend FMassRequirementAccessDetector;
 
 public:
-	FMassFragmentRequirements(){}
-	FMassFragmentRequirements(std::initializer_list<UScriptStruct*> InitList);
-	FMassFragmentRequirements(TConstArrayView<const UScriptStruct*> InitList);
+	FMassFragmentRequirements() = default;
+	explicit FMassFragmentRequirements(const TSharedPtr<FMassEntityManager>& EntityManager);
+	explicit FMassFragmentRequirements(const TSharedRef<FMassEntityManager>& EntityManager);
+
+	void Initialize(const TSharedRef<FMassEntityManager>& EntityManager);
 
 	FMassFragmentRequirements& AddRequirement(const UScriptStruct* FragmentType, const EMassFragmentAccess AccessMode, const EMassFragmentPresence Presence = EMassFragmentPresence::All)
 	{
+		checkf(bInitialized, TEXT("Modifying requirements before initialization is not supported."));
 		checkf(FragmentRequirements.FindByPredicate([FragmentType](const FMassFragmentRequirementDescription& Item){ return Item.StructType == FragmentType; }) == nullptr
 			, TEXT("Duplicated requirements are not supported. %s already present"), *GetNameSafe(FragmentType));
 		
@@ -203,6 +206,7 @@ public:
 	template<typename T>
 	FMassFragmentRequirements& AddRequirement(const EMassFragmentAccess AccessMode, const EMassFragmentPresence Presence = EMassFragmentPresence::All)
 	{
+		checkf(bInitialized, TEXT("Modifying requirements before initialization is not supported."));
 		checkf(FragmentRequirements.FindByPredicate([](const FMassFragmentRequirementDescription& Item) { return Item.StructType == T::StaticStruct(); }) == nullptr
 			, TEXT("Duplicated requirements are not supported. %s already present"), *T::StaticStruct()->GetName());
 
@@ -235,6 +239,7 @@ public:
 
 	void AddTagRequirement(const UScriptStruct& TagType, const EMassFragmentPresence Presence)
 	{
+		checkf(bInitialized, TEXT("Modifying requirements before initialization is not supported."));
 		checkf(int(Presence) != int(EMassFragmentPresence::MAX), TEXT("MAX presence is not a valid value for AddTagRequirement"));
 		switch (Presence)
 		{
@@ -257,6 +262,7 @@ public:
 	template<typename T>
 	FMassFragmentRequirements& AddTagRequirement(const EMassFragmentPresence Presence)
 	{
+		checkf(bInitialized, TEXT("Modifying requirements before initialization is not supported."));
 		checkf(int(Presence) != int(EMassFragmentPresence::MAX), TEXT("MAX presence is not a valid value for AddTagRequirement"));
 		static_assert(TIsDerivedFrom<T, FMassTag>::IsDerived, "Given struct doesn't represent a valid tag type. Make sure to inherit from FMassFragment or one of its child-types.");
 		switch (Presence)
@@ -294,6 +300,7 @@ public:
 	FMassFragmentRequirements& AddChunkRequirement(const EMassFragmentAccess AccessMode, const EMassFragmentPresence Presence = EMassFragmentPresence::All)
 	{
 		static_assert(TIsDerivedFrom<T, FMassChunkFragment>::IsDerived, "Given struct doesn't represent a valid chunk fragment type. Make sure to inherit from FMassChunkFragment or one of its child-types.");
+		checkf(bInitialized, TEXT("Modifying requirements before initialization is not supported."));
 		checkf(ChunkFragmentRequirements.FindByPredicate([](const FMassFragmentRequirementDescription& Item) { return Item.StructType == T::StaticStruct(); }) == nullptr
 			, TEXT("Duplicated requirements are not supported. %s already present"), *T::StaticStruct()->GetName());
 		checkf(Presence != EMassFragmentPresence::Any, TEXT("\'Any\' is not a valid Presence value for AddChunkRequirement."));
@@ -320,6 +327,7 @@ public:
 	FMassFragmentRequirements& AddConstSharedRequirement(const EMassFragmentPresence Presence = EMassFragmentPresence::All)
 	{
 		static_assert(TIsDerivedFrom<T, FMassConstSharedFragment>::IsDerived, "Given struct doesn't represent a valid const shared fragment type. Make sure to inherit from FMassConstSharedFragment or one of its child-types.");
+		checkf(bInitialized, TEXT("Modifying requirements before initialization is not supported."));
 		checkf(ConstSharedFragmentRequirements.FindByPredicate([](const FMassFragmentRequirementDescription& Item) { return Item.StructType == T::StaticStruct(); }) == nullptr
 			, TEXT("Duplicated requirements are not supported. %s already present"), *T::StaticStruct()->GetName());
 		checkf(Presence != EMassFragmentPresence::Any, TEXT("\'Any\' is not a valid Presence value for AddConstSharedRequirement."));
@@ -350,6 +358,7 @@ public:
 			return *this;
 		}
 
+		checkf(bInitialized, TEXT("Modifying requirements before initialization is not supported."));
 		checkf(ConstSharedFragmentRequirements.FindByPredicate([FragmentType](const FMassFragmentRequirementDescription& Item) { return Item.StructType == FragmentType; }) == nullptr
 			, TEXT("Duplicated requirements are not supported. %s already present"), *FragmentType->GetName());
 		checkf(Presence != EMassFragmentPresence::Any, TEXT("\'Any\' is not a valid Presence value for AddConstSharedRequirement."));
@@ -376,6 +385,7 @@ public:
 	FMassFragmentRequirements& AddSharedRequirement(const EMassFragmentAccess AccessMode, const EMassFragmentPresence Presence = EMassFragmentPresence::All)
 	{
 		static_assert(TIsDerivedFrom<T, FMassSharedFragment>::IsDerived, "Given struct doesn't represent a valid shared fragment type. Make sure to inherit from FMassSharedFragment or one of its child-types.");
+		checkf(bInitialized, TEXT("Modifying requirements before initialization is not supported."));
 		checkf(SharedFragmentRequirements.FindByPredicate([](const FMassFragmentRequirementDescription& Item) { return Item.StructType == T::StaticStruct(); }) == nullptr
 			, TEXT("Duplicated requirements are not supported. %s already present"), *T::StaticStruct()->GetName());
 		checkf(Presence != EMassFragmentPresence::Any, TEXT("\'Any\' is not a valid Presence value for AddSharedRequirement."));
@@ -438,6 +448,7 @@ public:
 	const FMassConstSharedFragmentBitSet& GetRequiredOptionalConstSharedFragments() const { return RequiredOptionalConstSharedFragments; }
 	const FMassConstSharedFragmentBitSet& GetRequiredNoneConstSharedFragments() const { return RequiredNoneConstSharedFragments; }
 
+	bool IsInitialized() const { return bInitialized; }
 	bool IsEmpty() const;
 	bool HasPositiveRequirements() const { return bHasPositiveRequirements; }
 	bool HasNegativeRequirements() const { return bHasNegativeRequirements; }
@@ -461,7 +472,6 @@ protected:
 	void ConsumeIncrementalChangesCount() { IncrementalChangesCount = 0; }
 	bool HasIncrementalChanges() const{ return IncrementalChangesCount > 0; }
 
-protected:
 	friend FMassRequirementAccessDetector;
 
 	TArray<FMassFragmentRequirementDescription> FragmentRequirements;
@@ -490,8 +500,10 @@ protected:
 	FMassConstSharedFragmentBitSet RequiredOptionalConstSharedFragments;
 	FMassConstSharedFragmentBitSet RequiredNoneConstSharedFragments;
 
+	TSharedPtr<FMassEntityManager> CachedEntityManager;
+
 private:
-	FORCEINLINE void CachePropreties() const;
+	void CacheProperties() const;
 	mutable uint16 bPropertiesCached : 1 = false;
 	mutable uint16 bHasPositiveRequirements : 1 = false;
 	mutable uint16 bHasNegativeRequirements : 1 = false;
@@ -500,12 +512,19 @@ private:
 	 * the optional elements will be accepted. Note that RequiredNone* requirements are handled separately and if specified 
 	 * still need to be satisfied.
 	 */
-	//mutable uint16 bOptionalsOnly : 1 = false;
 	mutable uint16 bHasOptionalRequirements : 1 = false;
 
+	uint16 bInitialized : 1 = false;
 	uint16 IncrementalChangesCount = 0;
 
 	bool bRequiresGameThreadExecution = false;
+
+public:
+	UE_DEPRECATED(5.6, "This type of FMassFragmentRequirements is no longer supported. Use one of the other constructors instead.")
+	FMassFragmentRequirements(std::initializer_list<UScriptStruct*> InitList);
+
+	UE_DEPRECATED(5.6, "This type of FMassFragmentRequirements is no longer supported. Use one of the other constructors instead.")
+	FMassFragmentRequirements(TConstArrayView<const UScriptStruct*> InitList);
 };
 
 template<>
