[Mass] Archetype/entity grouping

Added capability to group entities/archetypes with groups created at runtime. Groups can be used by EntityQueries to control the order of archetype processing. An example of when that could be useful would be a hierarchical structure where parents get processed before their children.

Current limitations:
* supported only in non-parallel ForEachEntityChunk approaches.
* no way to create entities with groups assigned at the very start - can only group entities after their creation (resulting in additional archetype creation and entity move)
* grouping is not "observable" via MassObservers
All to be addressed in near future

diff --git a/Engine/Source/Developer/MassEntityTestSuite/Private/MassArchetypeGroupTest.cpp b/Engine/Source/Developer/MassEntityTestSuite/Private/MassArchetypeGroupTest.cpp
new file mode 100644
--- /dev/null
+++ b/Engine/Source/Developer/MassEntityTestSuite/Private/MassArchetypeGroupTest.cpp
@@ -0,0 +1,292 @@
+// Copyright Epic Games, Inc. All Rights Reserved.
+
+#include "AITestsCommon.h"
+#include "MassEntityTestTypes.h"
+#include "MassArchetypeGroup.h"
+#include "MassExecutionContext.h"
+#include "Algo/Compare.h"
+
+#define LOCTEXT_NAMESPACE "MassTest"
+
+UE_DISABLE_OPTIMIZATION_SHIP
+
+namespace UE::Mass::Test
+{
+
+using namespace UE::Mass;
+
+auto DescendingGroupSorter = [](const FArchetypeGroupID A, const FArchetypeGroupID B)
+	{
+		return A > B;
+	};
+
+struct FArchetypeGroup_SingleLevelQuery : FEntityTestBase
+{
+	virtual bool InstantTest() override
+	{
+		const int32 NumEntities = 100;
+
+		FArchetypeGroupType GroupType = EntityManager->FinOrAddArchetypeGroupType(TEXT("TestGroup"));
+
+		TArray<FMassEntityHandle> Entities;
+		TArray<FMassEntityHandle> VerifyEntities;
+
+		EntityManager->BatchCreateEntities(IntsArchetype, NumEntities, Entities);
+
+		// verify assumptions regarding order of processing
+		FMassEntityQuery EntityQuery(EntityManager);
+		EntityQuery.AddRequirement<FTestFragment_Int>(EMassFragmentAccess::ReadOnly);
+
+		FMassExecutionContext ExecutionContext = EntityManager->CreateExecutionContext(0);
+		EntityQuery.ForEachEntityChunk(*EntityManager, ExecutionContext, [&VerifyEntities](const FMassExecutionContext& Context)
+		{
+			VerifyEntities.Append(Context.GetEntities());
+		});
+
+		AITEST_TRUE("Assumptions re order of entity processing is correct", Algo::Compare(Entities, VerifyEntities));
+
+		// we're assigning each entity an individual group.
+		for (int32 EntityIndex = 0; EntityIndex < Entities.Num(); ++EntityIndex)
+		{
+			EntityManager->BatchGroupEntities(FArchetypeGroupHandle(GroupType, EntityIndex), MakeArrayView(&Entities[EntityIndex], 1));
+		}
+
+		// by default the groups will be sorted in the ascending order, so the order of entities being processed should not change
+		EntityQuery.GroupBy(GroupType);
+
+		VerifyEntities.Reset();
+		EntityQuery.ForEachEntityChunk(*EntityManager, ExecutionContext, [&VerifyEntities](const FMassExecutionContext& Context)
+		{
+			VerifyEntities.Append(Context.GetEntities());
+		});
+		AITEST_TRUE("After grouping the order is expected to remain the same", Algo::Compare(Entities, VerifyEntities));
+
+		EntityQuery.ResetGrouping();
+		// this grouping should result in the reversed order of entity processing.
+		EntityQuery.GroupBy(GroupType, DescendingGroupSorter);
+
+		VerifyEntities.Reset();
+		EntityQuery.ForEachEntityChunk(*EntityManager, ExecutionContext, [&VerifyEntities](const FMassExecutionContext& Context)
+		{
+			// note that we're inserting at index 0 to end up with the ascending order of entities in the VerifyEntities array for ease of comparison
+			VerifyEntities.Insert(Context.GetEntities()[0], 0);
+		});
+		AITEST_TRUE("After descending grouping the order matches expectations", Algo::Compare(Entities, VerifyEntities));
+		
+		return true;
+	}
+};
+IMPLEMENT_AI_INSTANT_TEST(FArchetypeGroup_SingleLevelQuery, "System.Mass.Archetype.Group.SingleLevelQuery");
+
+struct FArchetypeGroup_MultiLevelQuery : FEntityTestBase
+{
+	virtual bool InstantTest() override
+	{
+		constexpr int32 GroupsNum = 3;
+		constexpr int32 NumEntities = 16;
+		const int32 GroupSizes[GroupsNum] = { 8, 4, 2 };
+		TArray<FArchetypeGroupType> GroupTypes;
+		for (int32 GroupIndex = 0; GroupIndex < GroupsNum; ++GroupIndex)
+		{
+			const FName GroupTypeName(*FString::Printf(TEXT("TestGroup_%d"), GroupIndex));
+			GroupTypes.Add(EntityManager->FinOrAddArchetypeGroupType(GroupTypeName));
+		}
+		
+		AITEST_NOT_EQUAL(TEXT("Two differently named group types are not expected to be equal"), GroupTypes[0], GroupTypes[1]);
+
+		TArray<FMassEntityHandle> Entities;
+		TArray<FMassEntityHandle> VerifyEntities;
+		EntityManager->BatchCreateEntities(IntsArchetype, NumEntities, Entities);
+		const int32 FirstEntityIndex = Entities[0].Index;
+
+		auto GroupAssigner = [&](auto&& Self, const int32 GroupIndex, TArrayView<FMassEntityHandle> EntitiesView) -> void
+		{
+			if (GroupIndex < GroupsNum)
+			{
+				const FArchetypeGroupType& GroupType = GroupTypes[GroupIndex];
+				const int32 GroupSize = GroupSizes[GroupIndex];
+				int32 SubGroupIndex = 0;
+				for (int32 EntityIndex = 0; EntityIndex < EntitiesView.Num(); EntityIndex += GroupSize, ++SubGroupIndex)
+				{
+					TArrayView<FMassEntityHandle> SubEntities = MakeArrayView(&EntitiesView[EntityIndex], GroupSize);
+					EntityManager->BatchGroupEntities(FArchetypeGroupHandle(GroupType, SubGroupIndex), SubEntities);
+					Self(Self, GroupIndex + 1, SubEntities);
+				}
+			}
+		};
+
+		GroupAssigner(GroupAssigner, 0, Entities);
+
+		// at this point we have:
+		// Entity	| TestGroup_0 	| TestGroup_1 	| TestGroup_2
+		//  0		|		0		|		0		|		0
+		//  1		|		0		|		0		|		0
+		//  2		|		0		|		0		|		1
+		//  3		|		0		|		0		|		1
+		//  4		|		0		|		1		|		0
+		//  5		|		0		|		1		|		0
+		//  6		|		0		|		1		|		1
+		//  7		|		0		|		1		|		1
+		//  8		|		1		|		0		|		0
+		//  9		|		1		|		0		|		0
+		// 10		|		1		|		0		|		1
+		// 11		|		1		|		0		|		1
+		// 12		|		1		|		1		|		0
+		// 13		|		1		|		1		|		0
+		// 14		|		1		|		1		|		1
+		// 15		|		1		|		1		|		1
+
+		FMassExecutionContext ExecutionContext = EntityManager->CreateExecutionContext(0);
+		FMassEntityQuery EntityQuery(EntityManager);
+		EntityQuery.AddRequirement<FTestFragment_Int>(EMassFragmentAccess::ReadOnly);
+		for (int32 GroupIndex = 0; GroupIndex < GroupsNum; ++GroupIndex)
+		{
+			const FArchetypeGroupType& GroupType = GroupTypes[GroupIndex];
+			EntityQuery.GroupBy(GroupType);
+		}
+
+		VerifyEntities.Reset();
+		EntityQuery.ForEachEntityChunk(*EntityManager, ExecutionContext, [&VerifyEntities](const FMassExecutionContext& Context)
+		{
+			VerifyEntities.Append(Context.GetEntities());
+		});
+
+		// by default we don't expect the order of entities processed to change
+		AITEST_TRUE("After grouping with the default sorting the order matches expectations", Algo::Compare(Entities, VerifyEntities));
+		{
+			// now we're going to reverse the sorting of group TestGroup_1
+			EntityQuery.ResetGrouping();
+			EntityQuery.GroupBy(GroupTypes[0]);
+			EntityQuery.GroupBy(GroupTypes[1], DescendingGroupSorter);
+			EntityQuery.GroupBy(GroupTypes[2]);
+
+			// here's what we expect to see with the TestGroup_1 group sorted in descending order
+			const int32 ExpectedEntityIndices[NumEntities] = { 4, 5, 6, 7, 0, 1, 2, 3, 12, 13, 14, 15, 8, 9, 10, 11 };
+
+			VerifyEntities.Reset();
+			EntityQuery.ForEachEntityChunk(*EntityManager, ExecutionContext, [&VerifyEntities](const FMassExecutionContext& Context)
+			{
+				VerifyEntities.Append(Context.GetEntities());
+			});
+
+			for (int32 ResultIndex = 0; ResultIndex < VerifyEntities.Num(); ++ResultIndex)
+			{
+				AITEST_EQUAL("Sorting TestGroup_1 in a descending order yields the expected results", VerifyEntities[ResultIndex], Entities[ExpectedEntityIndices[ResultIndex]]);
+			}
+		}
+		{
+			// now we're going to reverse the sorting of the last group, TestGroup_2
+			EntityQuery.ResetGrouping();
+			EntityQuery.GroupBy(GroupTypes[0]);
+			EntityQuery.GroupBy(GroupTypes[1]);
+			EntityQuery.GroupBy(GroupTypes[2], DescendingGroupSorter);
+
+			// here's what we expect to see with the TestGroup_1 group sorted in descending order
+			const int32 ExpectedEntityIndices[NumEntities] = { 2, 3, 0, 1, 6, 7, 4, 5, 10, 11, 8, 9, 14, 15, 12, 13 };
+
+			VerifyEntities.Reset();
+			EntityQuery.ForEachEntityChunk(*EntityManager, ExecutionContext, [&VerifyEntities](const FMassExecutionContext& Context)
+			{
+				VerifyEntities.Append(Context.GetEntities());
+			});
+
+			for (int32 ResultIndex = 0; ResultIndex < VerifyEntities.Num(); ++ResultIndex)
+			{
+				AITEST_EQUAL("Sorting TestGroup_2 in a descending order yields the expected results", VerifyEntities[ResultIndex], Entities[ExpectedEntityIndices[ResultIndex]]);
+			}
+		}
+		{
+			// now we're going to reverse the sorting of the TestGroup_0 and TestGroup_1, and put the latter first
+			EntityQuery.ResetGrouping();
+			EntityQuery.GroupBy(GroupTypes[1], DescendingGroupSorter);
+			EntityQuery.GroupBy(GroupTypes[0], DescendingGroupSorter);
+			EntityQuery.GroupBy(GroupTypes[2]);
+
+			// here's what we expect to see with the TestGroup_1 group sorted in descending order
+			const int32 ExpectedEntityIndices[NumEntities] = { 12, 13, 14, 15, 4, 5, 6, 7, 8, 9, 10, 11, 0, 1, 2, 3 };
+
+			VerifyEntities.Reset();
+			EntityQuery.ForEachEntityChunk(*EntityManager, ExecutionContext, [&VerifyEntities](const FMassExecutionContext& Context)
+			{
+				VerifyEntities.Append(Context.GetEntities());
+			});
+
+			for (int32 ResultIndex = 0; ResultIndex < VerifyEntities.Num(); ++ResultIndex)
+			{
+				AITEST_EQUAL("Sorting by TestGroup_2 and TestGroup_1 in a descending order yields the expected results", VerifyEntities[ResultIndex], Entities[ExpectedEntityIndices[ResultIndex]]);
+			}
+		}
+
+		// @TODO: test sorting when some entities are missing some group types - specifically the first one being used for sorting
+
+		return true;
+	}
+};
+IMPLEMENT_AI_INSTANT_TEST(FArchetypeGroup_MultiLevelQuery, "System.Mass.Archetype.Group.MultiLevelQuery");
+
+/** This test ensures the order in which a given entity is grouped in doesn't matter. */
+struct FArchetypeGroup_ApplicationOrder : FEntityTestBase
+{
+	virtual bool InstantTest() override
+	{
+		uint32 ArbitraryNumber = 1677;
+
+		FArchetypeGroupType GroupType1 = EntityManager->FinOrAddArchetypeGroupType(TEXT("TestGroup1"));
+		FArchetypeGroupType GroupType2 = EntityManager->FinOrAddArchetypeGroupType(TEXT("TestGroup2"));
+
+		TArray<FMassEntityHandle> Entities;
+		EntityManager->BatchCreateEntities(IntsArchetype, 2, Entities);
+
+		EntityManager->BatchGroupEntities(FArchetypeGroupHandle(GroupType1, ArbitraryNumber), MakeArrayView(&Entities[0], 1));
+		EntityManager->BatchGroupEntities(FArchetypeGroupHandle(GroupType2, ArbitraryNumber), MakeArrayView(&Entities[0], 1));
+
+		EntityManager->BatchGroupEntities(FArchetypeGroupHandle(GroupType2, ArbitraryNumber), MakeArrayView(&Entities[1], 1));
+		EntityManager->BatchGroupEntities(FArchetypeGroupHandle(GroupType1, ArbitraryNumber), MakeArrayView(&Entities[1], 1));
+
+		FMassArchetypeHandle Archetype0 = EntityManager->GetArchetypeForEntity(Entities[0]);
+		FMassArchetypeHandle Archetype1 = EntityManager->GetArchetypeForEntity(Entities[1]);
+
+		AITEST_EQUAL("Final archetype target is independent of order of entity grouping", Archetype0, Archetype1);
+
+		return true;
+	}
+};
+IMPLEMENT_AI_INSTANT_TEST(FArchetypeGroup_ApplicationOrder, "System.Mass.Archetype.Group.ApplicationOrder");
+
+/** Ensuring that grouping an entity moves it to a different archetype*/
+struct FArchetypeGroup_NewArchetype : FEntityTestBase
+{
+	virtual bool InstantTest() override
+	{
+		uint32 ArbitraryNumber = 1677;
+
+		FArchetypeGroupType GroupType = EntityManager->FinOrAddArchetypeGroupType(TEXT("TestGroup"));
+
+		const FMassEntityHandle Entity = EntityManager->CreateEntity(IntsArchetype);
+
+		EntityManager->BatchGroupEntities(FArchetypeGroupHandle(GroupType, ArbitraryNumber), MakeArrayView(&Entity, 1));
+		FMassArchetypeHandle ArchetypeGroup0 = EntityManager->GetArchetypeForEntity(Entity);
+		AITEST_NOT_EQUAL("Adding an entity to a group of any type makes it change archetypes", ArchetypeGroup0, IntsArchetype);
+
+		EntityManager->BatchGroupEntities(FArchetypeGroupHandle(GroupType, ArbitraryNumber + 1), MakeArrayView(&Entity, 1));
+		FMassArchetypeHandle ArchetypeGroup1 = EntityManager->GetArchetypeForEntity(Entity);
+		AITEST_NOT_EQUAL("Switching an entity to a different instance of the given group type makes it change archetypes", ArchetypeGroup0, ArchetypeGroup1);
+
+		EntityManager->BatchGroupEntities(FArchetypeGroupHandle(GroupType, ArbitraryNumber), MakeArrayView(&Entity, 1));
+		FMassArchetypeHandle ArchetypeGroup0B = EntityManager->GetArchetypeForEntity(Entity);
+		AITEST_EQUAL("Switching an entity to back to the original group instance makes it change archetypes back to the previous one", ArchetypeGroup0, ArchetypeGroup0B);
+
+		EntityManager->BatchGroupEntities(FArchetypeGroupHandle(GroupType, ArbitraryNumber), MakeArrayView(&Entity, 1));
+		FMassArchetypeHandle ArchetypeGroup0C = EntityManager->GetArchetypeForEntity(Entity);
+		AITEST_EQUAL("Attempting to add an entity to a group it's already in results in noop", ArchetypeGroup0, ArchetypeGroup0C);
+
+		return true;
+	}
+};
+IMPLEMENT_AI_INSTANT_TEST(FArchetypeGroup_NewArchetype, "System.Mass.Archetype.Group.NewArchetype");
+
+} // UE::Mass::Test
+
+UE_ENABLE_OPTIMIZATION_SHIP
+
+#undef LOCTEXT_NAMESPACE
diff --git a/Engine/Source/Runtime/MassEntity/Internal/MassArchetypeData.h b/Engine/Source/Runtime/MassEntity/Internal/MassArchetypeData.h
--- a/Engine/Source/Runtime/MassEntity/Internal/MassArchetypeData.h
+++ b/Engine/Source/Runtime/MassEntity/Internal/MassArchetypeData.h
@@ -4,6 +4,7 @@
 
 #include "MassEntityHandle.h"
 #include "MassArchetypeTypes.h"
+#include "MassArchetypeGroup.h"
 #if UE_ENABLE_INCLUDE_ORDER_DEPRECATED_IN_5_6
 #include "MassEntityManager.h"
 #endif // UE_ENABLE_INCLUDE_ORDER_DEPRECATED_IN_5_6
@@ -173,6 +174,8 @@ private:
 	
 	TMap<const UScriptStruct*, int32> FragmentIndexMap;
 
+	UE::Mass::FArchetypeGroups Groups;
+
 	int32 NumEntitiesPerChunk;
 	SIZE_T TotalBytesPerEntity = 0;
 	int32 EntityListOffsetWithinChunk;
@@ -215,15 +218,15 @@ public:
 		return GetSharedFragmentValues(Entity.Index);
 	}
 
+	const UE::Mass::FArchetypeGroups& GetGroups() const;
+	bool IsInGroup(const UE::Mass::FArchetypeGroupHandle GroupHandle) const;
+
 	/** Method to iterate on all the fragment types */
 	void ForEachFragmentType(TFunction< void(const UScriptStruct* /*FragmentType*/)> Function) const;
 	bool HasFragmentType(const UScriptStruct* FragmentType) const;
 	bool HasTagType(const UScriptStruct* FragmentType) const { check(FragmentType); return CompositionDescriptor.Tags.Contains(*FragmentType); }
 
-	bool IsEquivalent(const FMassArchetypeCompositionDescriptor& OtherCompositionDescriptor) const
-	{
-		return CompositionDescriptor.IsEquivalent(OtherCompositionDescriptor);
-	}
+	bool IsEquivalent(const FMassArchetypeCompositionDescriptor& OtherCompositionDescriptor, const UE::Mass::FArchetypeGroups& OtherGroups) const;
 
 	void Initialize(const FMassEntityManager& EntityManager, const FMassArchetypeCompositionDescriptor& InCompositionDescriptor, const uint32 ArchetypeDataVersion);
 
@@ -231,7 +234,8 @@ public:
 	 * A special way of initializing an archetype resulting in a copy of BaseArchetype's setup with OverrideTags
 	 * replacing original tags of BaseArchetype
 	 */
-	void InitializeWithSimilar(const FMassEntityManager& EntityManager, const FMassArchetypeData& BaseArchetype, FMassArchetypeCompositionDescriptor&& NewComposition, const uint32 ArchetypeDataVersion);
+	void InitializeWithSimilar(const FMassEntityManager& EntityManager, const FMassArchetypeData& BaseArchetype
+		, FMassArchetypeCompositionDescriptor&& NewComposition, UE::Mass::FArchetypeGroups&& InGroups, const uint32 ArchetypeDataVersion);
 
 	void AddEntity(FMassEntityHandle Entity, const FMassArchetypeSharedFragmentValues& InSharedFragmentValues);
 	void RemoveEntity(FMassEntityHandle Entity);
@@ -448,3 +452,21 @@ struct FMassArchetypeHelper
 	MASSENTITY_API static bool DoesArchetypeMatchRequirements(const FMassArchetypeData& Archetype, const FMassFragmentRequirements& Requirements);
 	MASSENTITY_API static bool DoesArchetypeMatchRequirements(const FMassArchetypeCompositionDescriptor& ArchetypeComposition, const FMassFragmentRequirements& Requirements);
 };
+
+//-----------------------------------------------------------------------------
+// INLINES
+//-----------------------------------------------------------------------------
+inline const UE::Mass::FArchetypeGroups& FMassArchetypeData::GetGroups() const
+{
+	return Groups;
+}
+
+inline bool FMassArchetypeData::IsInGroup(const UE::Mass::FArchetypeGroupHandle GroupHandle) const
+{
+	if (GroupHandle.IsValid())
+	{
+		UE::Mass::FArchetypeGroupID FoundGroupID = Groups.GetID(GroupHandle.GetGroupType());
+		return FoundGroupID.IsValid() && FoundGroupID == GroupHandle.GetGroupID();
+	}
+	return false;
+}
diff --git a/Engine/Source/Runtime/MassEntity/Private/MassArchetypeData.cpp b/Engine/Source/Runtime/MassEntity/Private/MassArchetypeData.cpp
--- a/Engine/Source/Runtime/MassEntity/Private/MassArchetypeData.cpp
+++ b/Engine/Source/Runtime/MassEntity/Private/MassArchetypeData.cpp
@@ -91,7 +91,8 @@ void FMassArchetypeData::Initialize(const FMassEntityManager& EntityManager, con
 	EntityListOffsetWithinChunk = 0;
 }
 
-void FMassArchetypeData::InitializeWithSimilar(const FMassEntityManager& EntityManager, const FMassArchetypeData& BaseArchetype, FMassArchetypeCompositionDescriptor&& NewComposition, const uint32 ArchetypeDataVersion)
+void FMassArchetypeData::InitializeWithSimilar(const FMassEntityManager& EntityManager, const FMassArchetypeData& BaseArchetype
+	, FMassArchetypeCompositionDescriptor&& NewComposition, UE::Mass::FArchetypeGroups&& InGroups, const uint32 ArchetypeDataVersion)
 {
 	checkf(IsInitialized() == false, TEXT("Trying to InitializeWithSimilar but this archetype has already been initialized"));
 
@@ -112,6 +113,8 @@ void FMassArchetypeData::InitializeWithSimilar(const FMassEntityManager& EntityM
 	}
 	ChunkFragmentsTemplate = BaseArchetype.ChunkFragmentsTemplate;
 
+	Groups.Set(Forward<UE::Mass::FArchetypeGroups>(InGroups));
+
 	EntityListOffsetWithinChunk = 0;
 }
 
@@ -1359,6 +1362,11 @@ void FMassArchetypeData::BatchSetFragmentValues(TConstArrayView<FMassArchetypeEn
 	}
 }
 
+bool FMassArchetypeData::IsEquivalent(const FMassArchetypeCompositionDescriptor& OtherCompositionDescriptor, const UE::Mass::FArchetypeGroups& OtherGroups) const
+{
+	return CompositionDescriptor.IsEquivalent(OtherCompositionDescriptor) && Groups == OtherGroups;
+}
+
 //-----------------------------------------------------------------------------
 // FMassArchetypeHelper
 //-----------------------------------------------------------------------------
diff --git a/Engine/Source/Runtime/MassEntity/Private/MassArchetypeGroup.cpp b/Engine/Source/Runtime/MassEntity/Private/MassArchetypeGroup.cpp
new file mode 100644
--- /dev/null
+++ b/Engine/Source/Runtime/MassEntity/Private/MassArchetypeGroup.cpp
@@ -0,0 +1,46 @@
+// Copyright Epic Games, Inc. All Rights Reserved.
+
+#include "MassArchetypeGroup.h"
+
+
+namespace UE::Mass
+{
+	//-----------------------------------------------------------------------------
+	// FArchetypeGroups
+	//-----------------------------------------------------------------------------
+	bool FArchetypeGroups::operator==(const FArchetypeGroups& OtherGroups) const
+	{
+		return IDContainer == OtherGroups.IDContainer;
+	}
+
+	void FArchetypeGroups::Set(FArchetypeGroups&& InGroups)
+	{
+		IDContainer = MoveTemp(InGroups.IDContainer);
+	}
+
+	void FArchetypeGroups::Add(FArchetypeGroupHandle GroupHandle)
+	{
+		const int32 GroupTypeIndex = static_cast<int32>(GroupHandle.GetGroupType());
+		if (IDContainer.IsValidIndex(GroupTypeIndex))
+		{
+			IDContainer[GroupTypeIndex] = GroupHandle.GetGroupID();
+		}
+		else
+		{
+			IDContainer.EmplaceAt(GroupTypeIndex, GroupHandle.GetGroupID());
+		}
+	}
+
+	uint32 GetTypeHash(const FArchetypeGroups& Instance)
+	{
+		uint32 Hash = 0;
+
+		for (auto It = Instance.IDContainer.CreateConstIterator(); It; ++It)
+		{
+			const FArchetypeGroupHandle GroupHandle(FArchetypeGroupType(It.GetIndex()), Instance.IDContainer[It.GetIndex()]);
+			Hash = HashCombine(Hash, GetTypeHash(GroupHandle));
+		}
+
+		return Hash;
+	}
+} // namespace UE::Mass
diff --git a/Engine/Source/Runtime/MassEntity/Private/MassEntityManager.cpp b/Engine/Source/Runtime/MassEntity/Private/MassEntityManager.cpp
--- a/Engine/Source/Runtime/MassEntity/Private/MassEntityManager.cpp
+++ b/Engine/Source/Runtime/MassEntity/Private/MassEntityManager.cpp
@@ -410,21 +410,21 @@ FMassArchetypeHandle FMassEntityManager::GetOrCreateSuitableArchetype(const FMas
 		FMassArchetypeCompositionDescriptor NewDescriptor = ArchetypeData.GetCompositionDescriptor();
 		NewDescriptor.SharedFragments = SharedFragmentBitSet;
 		NewDescriptor.ConstSharedFragments = ConstSharedFragmentBitSet;
-		return CreateArchetype(NewDescriptor);
+		return CreateArchetype(NewDescriptor, CreationParams);
 	}
 	return ArchetypeHandle;
 }
 
 FMassArchetypeHandle FMassEntityManager::CreateArchetype(const FMassArchetypeCompositionDescriptor& Composition, const FMassArchetypeCreationParams& CreationParams)
 {
-	const uint32 TypeHash = Composition.CalculateHash();
+	const uint32 TypeHash = HashCombine(Composition.CalculateHash(), GetTypeHash(UE::Mass::FArchetypeGroups()));
 
 	TArray<TSharedPtr<FMassArchetypeData>>& HashRow = FragmentHashToArchetypeMap.FindOrAdd(TypeHash);
 
 	TSharedPtr<FMassArchetypeData> ArchetypeDataPtr;
 	for (const TSharedPtr<FMassArchetypeData>& Ptr : HashRow)
 	{
-		if (Ptr->IsEquivalent(Composition))
+		if (Ptr->IsEquivalent(Composition, /*Groups=*/{}))
 		{
 #if WITH_MASSENTITY_DEBUG
 			// Keep track of all names for this archetype.
@@ -476,7 +476,7 @@ FMassArchetypeHandle FMassEntityManager::InternalCreateSimilarArchetype(const TS
 		, SourceArchetypeRef.GetChunkFragmentBitSet()
 		, SourceArchetypeRef.GetSharedFragmentBitSet()
 		, SourceArchetypeRef.GetConstSharedFragmentBitSet());
-	return InternalCreateSimilarArchetype(SourceArchetypeRef, MoveTemp(NewComposition));
+	return InternalCreateSimilarArchetype(SourceArchetypeRef, MoveTemp(NewComposition), {});
 }
 
 FMassArchetypeHandle FMassEntityManager::InternalCreateSimilarArchetype(const TSharedPtr<FMassArchetypeData>& SourceArchetype, const FMassFragmentBitSet& OverrideFragments)
@@ -488,19 +488,27 @@ FMassArchetypeHandle FMassEntityManager::InternalCreateSimilarArchetype(const TS
 		, SourceArchetypeRef.GetChunkFragmentBitSet()
 		, SourceArchetypeRef.GetSharedFragmentBitSet()
 		, SourceArchetypeRef.GetConstSharedFragmentBitSet());
-	return InternalCreateSimilarArchetype(SourceArchetypeRef, MoveTemp(NewComposition));
+	return InternalCreateSimilarArchetype(SourceArchetypeRef, MoveTemp(NewComposition), {});
 }
 
-FMassArchetypeHandle FMassEntityManager::InternalCreateSimilarArchetype(const FMassArchetypeData& SourceArchetypeRef, FMassArchetypeCompositionDescriptor&& NewComposition)
+FMassArchetypeHandle FMassEntityManager::InternalCreateSimilarArchetype(const TSharedPtr<FMassArchetypeData>& SourceArchetype, UE::Mass::FArchetypeGroups&& GroupsOverride)
 {
-	const uint32 TypeHash = NewComposition.CalculateHash();
+	checkSlow(SourceArchetype.IsValid());
+	const FMassArchetypeData& SourceArchetypeRef = *SourceArchetype.Get();
+	FMassArchetypeCompositionDescriptor NewComposition = SourceArchetype->GetCompositionDescriptor();
+	return InternalCreateSimilarArchetype(SourceArchetypeRef, MoveTemp(NewComposition), Forward<UE::Mass::FArchetypeGroups>(GroupsOverride));
+}
+
+FMassArchetypeHandle FMassEntityManager::InternalCreateSimilarArchetype(const FMassArchetypeData& SourceArchetypeRef, FMassArchetypeCompositionDescriptor&& NewComposition, UE::Mass::FArchetypeGroups&& Groups)
+{
+	const uint32 TypeHash = HashCombine(NewComposition.CalculateHash(), GetTypeHash(Groups));
 
 	TArray<TSharedPtr<FMassArchetypeData>>& HashRow = FragmentHashToArchetypeMap.FindOrAdd(TypeHash);
 
 	TSharedPtr<FMassArchetypeData> ArchetypeDataPtr;
 	for (const TSharedPtr<FMassArchetypeData>& Ptr : HashRow)
 	{
-		if (Ptr->IsEquivalent(NewComposition))
+		if (Ptr->IsEquivalent(NewComposition, Groups))
 		{
 			ArchetypeDataPtr = Ptr;
 			break;
@@ -514,7 +522,7 @@ FMassArchetypeHandle FMassEntityManager::InternalCreateSimilarArchetype(const FM
 
 		// Create a new archetype
 		FMassArchetypeData* NewArchetype = new FMassArchetypeData(FMassArchetypeCreationParams(SourceArchetypeRef));
-		NewArchetype->InitializeWithSimilar(*this, SourceArchetypeRef, MoveTemp(NewComposition), ArchetypeDataVersion);
+		NewArchetype->InitializeWithSimilar(*this, SourceArchetypeRef, MoveTemp(NewComposition), Forward<UE::Mass::FArchetypeGroups>(Groups), ArchetypeDataVersion);
 		NewArchetype->CopyDebugNamesFrom(SourceArchetypeRef);
 
 		ArchetypeDataPtr = HashRow.Add_GetRef(MakeShareable(NewArchetype));
@@ -926,6 +934,75 @@ void FMassEntityManager::BatchDestroyEntityChunks(TConstArrayView<FMassArchetype
 	}
 }
 
+void FMassEntityManager::BatchGroupEntities(const UE::Mass::FArchetypeGroupHandle GroupHandle, TConstArrayView<FMassArchetypeEntityCollection> Collections)
+{
+	if (GroupHandle.IsValid() == false)
+	{
+		UE_LOG(LogMass, Warning, TEXT("%hs called with an invalid GroupHandle"), __FUNCTION__);
+		return;
+	}
+
+	TArray<FMassEntityHandle> EntitiesBeingMoved;
+
+	for (const FMassArchetypeEntityCollection& EntityCollection : Collections)
+	{
+		if (EntityCollection.GetArchetype().IsValid())
+		{
+			FMassArchetypeData& CurrentArchetype = FMassArchetypeHelper::ArchetypeDataFromHandleChecked(EntityCollection.GetArchetype());
+			if (CurrentArchetype.IsInGroup(GroupHandle) == false)
+			{
+				UE::Mass::FArchetypeGroups NewGroups = CurrentArchetype.GetGroups();
+				NewGroups.Add(GroupHandle);
+
+				const FMassArchetypeHandle NewArchetypeHandle = InternalCreateSimilarArchetype(EntityCollection.GetArchetype().DataPtr, MoveTemp(NewGroups));
+
+				EntitiesBeingMoved.Reset();
+				CurrentArchetype.BatchMoveEntitiesToAnotherArchetype(EntityCollection, *NewArchetypeHandle.DataPtr.Get(), EntitiesBeingMoved
+				// we need something like the following to support observers
+				//, bTagsAddedAreObserved ? &NewArchetypeEntityRanges : nullptr
+				);
+
+				for (const FMassEntityHandle& Entity : EntitiesBeingMoved)
+				{
+					check(GetEntityStorageInterface().IsValidIndex(Entity.Index));
+					GetEntityStorageInterface().SetArchetypeFromShared(Entity.Index, NewArchetypeHandle.DataPtr);
+				}
+			}
+		}
+	}
+}
+
+void FMassEntityManager::BatchGroupEntities(const UE::Mass::FArchetypeGroupHandle GroupHandle, TConstArrayView<FMassEntityHandle> InEntities)
+{
+	TArray<FMassArchetypeEntityCollection> Collections;
+	UE::Mass::Utils::CreateEntityCollections(*this, InEntities, FMassArchetypeEntityCollection::FoldDuplicates, Collections);
+	BatchGroupEntities(GroupHandle, Collections);
+}
+
+UE::Mass::FArchetypeGroupType FMassEntityManager::FinOrAddArchetypeGroupType(const FName GroupName)
+{
+	const int32* FoundGroupIndex = GroupNameToTypeIndex.Find(GroupName);
+	if (LIKELY(FoundGroupIndex))
+	{
+		return UE::Mass::FArchetypeGroupType(*FoundGroupIndex);
+	}
+
+	const int32 NewGroupIndex = GroupTypes.Add(GroupName);
+	GroupNameToTypeIndex.Add(GroupName, NewGroupIndex);
+	return UE::Mass::FArchetypeGroupType(NewGroupIndex);
+}
+
+const UE::Mass::FArchetypeGroups& FMassEntityManager::GetGroupsForArchetype(const FMassArchetypeHandle& ArchetypeHandle) const
+{
+	if (ArchetypeHandle.IsValid() == false)
+	{
+		static UE::Mass::FArchetypeGroups DummyGroups;
+		return DummyGroups;
+	}
+
+	return ArchetypeHandle.DataPtr->GetGroups();
+}
+
 void FMassEntityManager::AddFragmentToEntity(FMassEntityHandle Entity, const UScriptStruct* FragmentType)
 {
 	checkf(FragmentType, TEXT("Null fragment type passed in to %hs"), __FUNCTION__);
@@ -2107,3 +2184,8 @@ const FMassArchetypeEntityCollection& FMassEntityManager::FEntityCreationContext
 	static FMassArchetypeEntityCollection EmptyCollection;
 	return EntityCollections.Num() ? EntityCollections[0] : EmptyCollection;
 }
+
+FMassArchetypeHandle FMassEntityManager::InternalCreateSimilarArchetype(const FMassArchetypeData& SourceArchetypeRef, FMassArchetypeCompositionDescriptor&& NewComposition)
+{
+	return InternalCreateSimilarArchetype(SourceArchetypeRef, Forward<FMassArchetypeCompositionDescriptor>(NewComposition), {});
+}
diff --git a/Engine/Source/Runtime/MassEntity/Private/MassEntityQuery.cpp b/Engine/Source/Runtime/MassEntity/Private/MassEntityQuery.cpp
--- a/Engine/Source/Runtime/MassEntity/Private/MassEntityQuery.cpp
+++ b/Engine/Source/Runtime/MassEntity/Private/MassEntityQuery.cpp
@@ -102,6 +102,8 @@ void FMassEntityQuery::CacheArchetypes(const FMassEntityManager& InEntityManager
 		bUpdateArchetypes = true;
 		EntitySubsystemHash = InEntityManagerHash;
 		ValidArchetypes.Reset();
+		OrderedArchetypeIndices.Reset();
+		CachedGroupIDs.Reset();
 		LastUpdatedArchetypeDataVersion = 0;
 		ArchetypeFragmentMapping.Reset();
 
@@ -167,6 +169,15 @@ void FMassEntityQuery::CacheArchetypes(const FMassEntityManager& InEntityManager
 					ArchetypeData.GetRequirementsSharedFragmentMapping(SharedFragmentRequirements, ArchetypeFragmentMapping[i].SharedFragments);
 				}
 			}
+
+			if (GroupSortingSteps.Num())
+			{
+				SortArchetypes(FirstNewArchetype);
+			}
+			else
+			{
+				BuildOrderedArchetypeIndices(FirstNewArchetype);
+			}
 		}
 	}
 }
@@ -256,11 +267,12 @@ void FMassEntityQuery::ForEachEntityChunk(FMassEntityManager& EntityManager, FMa
 			// it's important to set requirements after caching archetypes due to that call potentially sorting the requirements and the order is relevant here.
 			ExecutionContext.SetFragmentRequirements(*this);
 
-			for (int i = 0; i < ValidArchetypes.Num(); ++i)
+			checkSlow(OrderedArchetypeIndices.Num() == ValidArchetypes.Num());
+			for (const int32 ArchetypeIndex : OrderedArchetypeIndices)
 			{
-				const FMassArchetypeHandle& ArchetypeHandle = ValidArchetypes[i];
+				const FMassArchetypeHandle& ArchetypeHandle = ValidArchetypes[ArchetypeIndex];
 				FMassArchetypeData& ArchetypeData = FMassArchetypeHelper::ArchetypeDataFromHandleChecked(ArchetypeHandle);
-				ArchetypeData.ExecuteFunction(ExecutionContext, ExecuteFunction, ArchetypeFragmentMapping[i], ChunkCondition);
+				ArchetypeData.ExecuteFunction(ExecutionContext, ExecuteFunction, ArchetypeFragmentMapping[ArchetypeIndex], ChunkCondition);
 				ExecutionContext.ClearFragmentViews();
 			}
 		}
@@ -493,6 +505,119 @@ bool FMassEntityQuery::HasMatchingEntities(FMassEntityManager& InEntityManager)
 	return false;
 }
 
+void FMassEntityQuery::GroupBy(UE::Mass::FArchetypeGroupType GroupType)
+{
+	GroupBy(GroupType, [](const UE::Mass::FArchetypeGroupID A, const UE::Mass::FArchetypeGroupID B)
+	{
+		return A < B;
+	});
+}
+
+void FMassEntityQuery::GroupBy(UE::Mass::FArchetypeGroupType GroupType, const TFunction<bool(const UE::Mass::FArchetypeGroupID, const UE::Mass::FArchetypeGroupID)>& Predicate)
+{
+	GroupSortingSteps.Emplace(GroupType, Predicate);
+	IncrementChangeCounter();
+}
+
+void FMassEntityQuery::ResetGrouping()
+{
+	GroupSortingSteps.Reset();
+	IncrementChangeCounter();
+}
+
+void FMassEntityQuery::BuildOrderedArchetypeIndices(const int32 StartingIndex)
+{
+	if (StartingIndex == 0)
+	{
+		OrderedArchetypeIndices.Reset();
+	}
+	OrderedArchetypeIndices.SetNumUninitialized(ValidArchetypes.Num());
+	for (int32 ArchetypeIndex = StartingIndex; ArchetypeIndex < ValidArchetypes.Num(); ++ArchetypeIndex)
+	{
+		OrderedArchetypeIndices[ArchetypeIndex] = ArchetypeIndex;
+	}
+}
+
+void FMassEntityQuery::SortArchetypes(const int32 FirstNewArchetypeIndex)
+{
+	if (GroupSortingSteps.Num() == 0)
+	{
+		return;
+	}
+
+	check(CachedEntityManager);
+
+	// first we need to cache the required group IDs from the new archetypes
+	CachedGroupIDs.SetNum(ValidArchetypes.Num());
+	OrderedArchetypeIndices.SetNumUninitialized(ValidArchetypes.Num());
+
+	for (int32 NewArchetypeIndex = FirstNewArchetypeIndex; NewArchetypeIndex < ValidArchetypes.Num(); ++NewArchetypeIndex)
+	{
+		OrderedArchetypeIndices[NewArchetypeIndex] = NewArchetypeIndex;
+
+		TArray<UE::Mass::FArchetypeGroupID>& ArchetypeGroupIDs = CachedGroupIDs[NewArchetypeIndex];
+		ArchetypeGroupIDs.Reserve(GroupSortingSteps.Num());
+
+		const UE::Mass::FArchetypeGroups& ArchetypeGroups = CachedEntityManager->GetGroupsForArchetype(ValidArchetypes[NewArchetypeIndex]);
+
+		for (const FArchetypeGroupingStep& Step : GroupSortingSteps)
+		{
+			// ArchetypeGroups.GetID will return InvalidArchetypeGroupID if the given group type is not in ArchetypeGroups.
+			// This is what we want.
+			ArchetypeGroupIDs.Add(ArchetypeGroups.GetID(Step.GroupType));
+		}
+	}
+
+	TConstArrayView<TArray<UE::Mass::FArchetypeGroupID>> GroupIDs = MakeConstArrayView(CachedGroupIDs);
+	TArray<TTuple<int32, int32>> Ranges;
+	Ranges.Add({0, OrderedArchetypeIndices.Num()});
+	int32 MaxRangeSize = OrderedArchetypeIndices.Num();
+	int32 Step = 0;
+	int32 RangesProcessed = 0;
+
+	while (Step < GroupSortingSteps.Num() && MaxRangeSize > 1)
+	{
+		const bool bLastIteration = (Step == GroupSortingSteps.Num() - 1);
+		const int32 RangesThisIteration = Ranges.Num();
+		MaxRangeSize = 0;
+		for (; RangesProcessed < RangesThisIteration; ++RangesProcessed)
+		{
+			const TTuple<int32, int32>& Range = Ranges[RangesProcessed];
+			const int32 RangeLength = Range.Get<1>() - Range.Get<0>();
+			MakeArrayView(&OrderedArchetypeIndices[Range.Get<0>()], RangeLength)
+				.Sort([Step, &GroupIDs, Sorter = GroupSortingSteps[Step].Predicate](const int32 A, const int32 B)
+				{
+					return Sorter(GroupIDs[A][Step], GroupIDs[B][Step]);
+				});
+
+			// figure out new ranges
+			if (bLastIteration == false)
+			{
+				int32 SubRangeStart = Range.Get<0>();
+				UE::Mass::FArchetypeGroupID PrevValue = GroupIDs[OrderedArchetypeIndices[SubRangeStart]][Step];
+				for (int32 Index = SubRangeStart + 1; Index < Range.Get<1>(); ++Index)
+				{
+					const UE::Mass::FArchetypeGroupID NewValue = GroupIDs[OrderedArchetypeIndices[Index]][Step];
+					if (NewValue != PrevValue)
+					{
+						PrevValue = NewValue;
+						Ranges.Push({SubRangeStart, Index});
+						MaxRangeSize = FMath::Max(MaxRangeSize, Index - SubRangeStart);
+						SubRangeStart = Index;
+					}
+				}
+
+				// the loop above doesn't create any ranges when there's no change in GroupID among processed archetypes.
+				// Similarly, it doesn't store the "last" range. We're addressing this now.
+				Ranges.Push({SubRangeStart, Range.Get<1>()});
+			}
+		}
+
+		check(MaxRangeSize >= 1 || bLastIteration);
+		++Step;
+	};
+}
+
 const FMassQueryRequirementIndicesMapping& FMassEntityQuery::GetRequirementsMappingForArchetype(const FMassArchetypeHandle ArchetypeHandle) const
 {
 	static const FMassQueryRequirementIndicesMapping FallbackEmptyMapping;
diff --git a/Engine/Source/Runtime/MassEntity/Public/MassArchetypeGroup.h b/Engine/Source/Runtime/MassEntity/Public/MassArchetypeGroup.h
new file mode 100644
--- /dev/null
+++ b/Engine/Source/Runtime/MassEntity/Public/MassArchetypeGroup.h
@@ -0,0 +1,144 @@
+// Copyright Epic Games, Inc. All Rights Reserved.
+
+#pragma once
+
+#include "Containers/ArrayView.h"
+#include "Containers/Map.h"
+
+
+namespace UE::Mass
+{
+	//-----------------------------------------------------------------------------
+	// FArchetypeGroupType
+	//-----------------------------------------------------------------------------
+	struct FArchetypeGroupType
+	{
+		explicit FArchetypeGroupType(int32 Value = INDEX_NONE)
+			: Identifier(Value)
+		{
+		}
+
+		friend uint32 GetTypeHash(const FArchetypeGroupType& Instance)
+		{
+			return GetTypeHash(Instance.Identifier);
+		}
+
+		bool operator==(const FArchetypeGroupType Other) const
+		{
+			return Identifier == Other.Identifier;
+		}
+
+		bool operator<(const FArchetypeGroupType Other) const
+		{
+			return Identifier < Other.Identifier;
+		}
+
+		explicit operator int32() const
+		{
+			return Identifier;
+		}
+
+		bool IsValid() const
+		{
+			return Identifier != INDEX_NONE;
+		}
+
+	private:
+		const int32 Identifier;
+	};
+
+
+	struct FArchetypeGroupID
+	{
+		constexpr static uint32 InvalidArchetypeGroupID = static_cast<uint32>(INDEX_NONE);
+
+		FArchetypeGroupID(const uint32 InID = InvalidArchetypeGroupID)
+			: ID(InID)
+		{
+			
+		}
+
+		bool IsValid() const
+		{
+			return ID != InvalidArchetypeGroupID;
+		}
+
+		operator int32() const
+		{
+			return static_cast<int32>(ID);
+		}
+
+	private:
+		uint32 ID;
+	};
+
+	//-----------------------------------------------------------------------------
+	// FArchetypeGroupHandle
+	//-----------------------------------------------------------------------------
+	struct FArchetypeGroupHandle
+	{
+		FArchetypeGroupHandle(const FArchetypeGroupType InGroupType, const FArchetypeGroupID InGroupID)
+			: GroupType(InGroupType), GroupID(InGroupID)
+		{
+		}
+
+		friend uint32 GetTypeHash(const FArchetypeGroupHandle& Instance)
+		{
+			const uint64 CombinedHandle = static_cast<uint64>(static_cast<int32>(Instance.GroupType)) << 32 | Instance.GroupID;
+			return GetTypeHash(CombinedHandle);
+		}
+
+		FArchetypeGroupType GetGroupType() const
+		{
+			return GroupType;
+		}
+
+		FArchetypeGroupID GetGroupID() const
+		{
+			return GroupID;
+		}
+
+		bool operator<(const FArchetypeGroupHandle Other) const
+		{
+			return GroupType < Other.GroupType || (GroupType == Other.GroupType && GroupID < Other.GroupID);
+		}
+
+		void UpdateID(const FArchetypeGroupHandle Other)
+		{
+			if (ensureMsgf(Other.GroupType == GroupType, TEXT("Updating ID is only supported for group handles of the same type")))
+			{
+				new (this) FArchetypeGroupHandle(GroupType, Other.GroupID);
+			}
+		}
+
+		bool IsValid() const
+		{
+			return GroupType.IsValid() && GroupID.IsValid();
+		}
+
+	private:
+		const FArchetypeGroupType GroupType;
+		const FArchetypeGroupID GroupID;
+	};
+
+	//-----------------------------------------------------------------------------
+	// FArchetypeGroups
+	//-----------------------------------------------------------------------------
+	struct FArchetypeGroups
+	{
+		bool operator==(const FArchetypeGroups& OtherGroups) const;
+		void Set(UE::Mass::FArchetypeGroups&& InGroups);
+		void Add(FArchetypeGroupHandle GroupHandle);
+
+		FArchetypeGroupID GetID(const FArchetypeGroupType GroupType) const
+		{
+			return IDContainer.IsValidIndex(static_cast<int32>(GroupType)) ? IDContainer[static_cast<int32>(GroupType)] : FArchetypeGroupID();
+		}
+
+		friend uint32 GetTypeHash(const FArchetypeGroups& Instance);
+
+	protected:
+		TSparseArray<FArchetypeGroupID> IDContainer;
+	};
+
+} // namespace UE::Mass
diff --git a/Engine/Source/Runtime/MassEntity/Public/MassEntityManager.h b/Engine/Source/Runtime/MassEntity/Public/MassEntityManager.h
--- a/Engine/Source/Runtime/MassEntity/Public/MassEntityManager.h
+++ b/Engine/Source/Runtime/MassEntity/Public/MassEntityManager.h
@@ -16,6 +16,7 @@
 #include "MassRequirementAccessDetector.h"
 #include "Templates/FunctionFwd.h"
 #include "MassEntityManagerStorage.h"
+#include "MassArchetypeGroup.h"
 
 
 struct FInstancedStruct;
@@ -346,6 +347,28 @@ public:
 	void BatchDestroyEntityChunks(const FMassArchetypeEntityCollection& Collection);
 	void BatchDestroyEntityChunks(TConstArrayView<FMassArchetypeEntityCollection> Collections);
 
+	/**
+	 * Assigns all entities indicated by Collections to a given archetype group.
+	 * Note that depending on their individual composition each entity can end up in a different archetype.
+	 * @paramm GroupHandle indicates the target group. Passing an invalid group handle will get logged as warning and ignored.
+	 */
+	void BatchGroupEntities(const UE::Mass::FArchetypeGroupHandle GroupHandle, TConstArrayView<FMassArchetypeEntityCollection> Collections);
+
+	/**
+	 * Assigns all entities indicated by InEntities to a given archetype group.
+	 * Note that depending on their individual composition each entity can end up in a different archetype.
+	 * @paramm GroupHandle indicates the target group. Passing an invalid group handle will get logged as warning and ignored.
+	 */
+	void BatchGroupEntities(const UE::Mass::FArchetypeGroupHandle GroupHandle, TConstArrayView<FMassEntityHandle> InEntities);
+
+	/**
+	 * Fetches FArchetypeGroupType instance (copy) associated with the given GroupName. A new group type
+	 * is created if GroupName has not been used in the past.f
+	 */
+	UE::Mass::FArchetypeGroupType FinOrAddArchetypeGroupType(const FName GroupName);
+
+	const UE::Mass::FArchetypeGroups& GetGroupsForArchetype(const FMassArchetypeHandle& ArchetypeHandle) const;
+
 	void AddFragmentToEntity(FMassEntityHandle Entity, const UScriptStruct* FragmentType);
 	void AddFragmentToEntity(FMassEntityHandle Entity, const UScriptStruct* FragmentType, const FStructInitializationCallback& Initializer);
 
@@ -1026,8 +1049,8 @@ protected:
 	 */
 	FMassArchetypeHandle InternalCreateSimilarArchetype(const TSharedPtr<FMassArchetypeData>& SourceArchetype, const FMassTagBitSet& OverrideTags);
 	FMassArchetypeHandle InternalCreateSimilarArchetype(const TSharedPtr<FMassArchetypeData>& SourceArchetype, const FMassFragmentBitSet& OverrideFragments);
-
-	FMassArchetypeHandle InternalCreateSimilarArchetype(const FMassArchetypeData& SourceArchetypeRef, FMassArchetypeCompositionDescriptor&& NewComposition);
+	FMassArchetypeHandle InternalCreateSimilarArchetype(const TSharedPtr<FMassArchetypeData>& SourceArchetype, UE::Mass::FArchetypeGroups&& GroupsOverride);
+	FMassArchetypeHandle InternalCreateSimilarArchetype(const FMassArchetypeData& SourceArchetypeRef, FMassArchetypeCompositionDescriptor&& NewComposition, UE::Mass::FArchetypeGroups&& GroupsOverride);
 
 	void InternalAppendFragmentsAndTagsToArchetypeCompositionDescriptor(FMassArchetypeCompositionDescriptor& InOutComposition,
 		TConstArrayView<const UScriptStruct*> FragmentsAndTagsList) const;
@@ -1185,6 +1208,10 @@ private:
 
 	FMassObserverManager ObserverManager;
 
+	TMap<const FName, const int32> GroupNameToTypeIndex;
+	// @todo we'll probably have some "GroupTypeInformation" here in the future
+	TArray<const FName> GroupTypes;
+
 #if WITH_MASSENTITY_DEBUG
 	FMassRequirementAccessDetector RequirementAccessDetector;
 	FString DebugName;
@@ -1198,4 +1225,10 @@ private:
 	bool bFirstCommandFlush = true;
 
 	FDelegateHandle OnPostForkHandle;
+
+	//-----------------------------------------------------------------------------
+	// DEPRECATED
+	//-----------------------------------------------------------------------------
+	UE_DEPRECATED(5.5, "This flavor of InternalCreateSimilarArchetype is deprecaded due to the introduction of archetype grouping. Use InternalCreateSimilarArchetype with a FArchetypeGroups parameter instead")
+	FMassArchetypeHandle InternalCreateSimilarArchetype(const FMassArchetypeData& SourceArchetypeRef, FMassArchetypeCompositionDescriptor&& NewComposition);
 };
diff --git a/Engine/Source/Runtime/MassEntity/Public/MassEntityQuery.h b/Engine/Source/Runtime/MassEntity/Public/MassEntityQuery.h
--- a/Engine/Source/Runtime/MassEntity/Public/MassEntityQuery.h
+++ b/Engine/Source/Runtime/MassEntity/Public/MassEntityQuery.h
@@ -9,6 +9,7 @@
 #endif // UE_ENABLE_INCLUDE_ORDER_DEPRECATED_IN_5_6
 #include "MassEntityTypes.h"
 #include "MassArchetypeTypes.h"
+#include "MassArchetypeGroup.h"
 #if UE_ENABLE_INCLUDE_ORDER_DEPRECATED_IN_5_6
 #include "MassExternalSubsystemTraits.h"
 #endif // UE_ENABLE_INCLUDE_ORDER_DEPRECATED_IN_5_6
@@ -108,6 +109,7 @@ public:
 	{
 		FMassFragmentRequirements::Reset();
 		FMassSubsystemRequirements::Reset();
+		ResetGrouping();
 		DirtyCachedData();
 	}
 
@@ -164,7 +166,11 @@ public:
 	void ClearChunkFilter() { ChunkCondition.Reset(); }
 
 	bool HasChunkFilter() const { return bool(ChunkCondition); }
-	
+
+	void GroupBy(UE::Mass::FArchetypeGroupType GroupType);
+	void GroupBy(UE::Mass::FArchetypeGroupType GroupType, const TFunction<bool(const UE::Mass::FArchetypeGroupID, const UE::Mass::FArchetypeGroupID)>& Predicate);
+	void ResetGrouping();
+
 #if WITH_ARCHETYPE_MATCH_OVERRIDE
 	template<typename T> requires TArchetypeMatchOverrideConcept<T>
 	void SetArchetypeMatchOverride(const T& Override);
@@ -195,6 +201,9 @@ private:
 		FMassExternalSubsystemBitSet MutableSubsystemsBitSet;
 	};
 
+	void SortArchetypes(const int32 FirstNewArchetype);
+	void BuildOrderedArchetypeIndices(const int32 StartingIndex = 0);
+
 	/** 
 	 * This function represents a condition that will be called for every chunk to be processed before the actual 
 	 * execution function is called. The chunk fragment requirements are already bound and ready to be used by the time 
@@ -206,8 +215,17 @@ private:
 	uint32 LastUpdatedArchetypeDataVersion = 0;
 
 	TArray<FMassArchetypeHandle> ValidArchetypes;
+	TArray<int32> OrderedArchetypeIndices;
 	TArray<FMassQueryRequirementIndicesMapping> ArchetypeFragmentMapping;
 
+	struct FArchetypeGroupingStep
+	{
+		UE::Mass::FArchetypeGroupType GroupType;
+		TFunction<bool(const UE::Mass::FArchetypeGroupID, const UE::Mass::FArchetypeGroupID)> Predicate;
+	};
+	TArray<FArchetypeGroupingStep> GroupSortingSteps;
+	TArray<TArray<UE::Mass::FArchetypeGroupID>> CachedGroupIDs;
+
 	/** 
 	 * Controls whether ParallelForEachEntityChunk created dedicated command buffer for each job. This is required 
 	 * to ensure thread safety. Disable by calling SetParallelCommandBufferEnabled(false) if execution function doesn't 
