[Mass] entity grouping improvements:

* Fixed the issue causing entities to lose their group information when moved to another archetype
* Added means of checking entity's group and removing an entity from a given group type
* Added unit-tests for the new functionality

diff --git a/Engine/Source/Developer/MassEntityTestSuite/Private/MassArchetypeGroupTest.cpp b/Engine/Source/Developer/MassEntityTestSuite/Private/MassArchetypeGroupTest.cpp
--- a/Engine/Source/Developer/MassEntityTestSuite/Private/MassArchetypeGroupTest.cpp
+++ b/Engine/Source/Developer/MassEntityTestSuite/Private/MassArchetypeGroupTest.cpp
@@ -253,6 +253,36 @@ struct FArchetypeGroup_ApplicationOrder : FEntityTestBase
 };
 IMPLEMENT_AI_INSTANT_TEST(FArchetypeGroup_ApplicationOrder, "System.Mass.Archetype.Group.ApplicationOrder");
 
+struct FArchetypeGroup_Equivalence : FEntityTestBase
+{
+	virtual bool InstantTest() override
+	{
+		uint32 ArbitraryNumber = 1677;
+
+		FArchetypeGroupType GroupType1 = EntityManager->FinOrAddArchetypeGroupType(TEXT("TestGroup1"));
+		FArchetypeGroupType GroupType2 = EntityManager->FinOrAddArchetypeGroupType(TEXT("TestGroup2"));
+
+		TArray<FMassEntityHandle> Entities;
+		EntityManager->BatchCreateEntities(IntsArchetype, 2, Entities);
+
+		// the first entity only gets the first group
+		EntityManager->BatchGroupEntities(FArchetypeGroupHandle(GroupType1, ArbitraryNumber), MakeArrayView(&Entities[0], 1));
+		
+		// the second one gets two and then gets removed from the latter group
+		EntityManager->BatchGroupEntities(FArchetypeGroupHandle(GroupType1, ArbitraryNumber), MakeArrayView(&Entities[1], 1));
+		EntityManager->BatchGroupEntities(FArchetypeGroupHandle(GroupType2, ArbitraryNumber), MakeArrayView(&Entities[1], 1));
+		EntityManager->RemoveEntityFromGroupType(Entities[1], GroupType2);
+
+		FMassArchetypeHandle Archetype0 = EntityManager->GetArchetypeForEntity(Entities[0]);
+		FMassArchetypeHandle Archetype1 = EntityManager->GetArchetypeForEntity(Entities[1]);
+
+		AITEST_EQUAL("The archetypes with groups created directly and with groups modifications", Archetype0, Archetype1);
+
+		return true;
+	}
+};
+IMPLEMENT_AI_INSTANT_TEST(FArchetypeGroup_Equivalence, "System.Mass.Archetype.Group.Equivalence");
+
 /** Ensuring that grouping an entity moves it to a different archetype*/
 struct FArchetypeGroup_NewArchetype : FEntityTestBase
 {
@@ -285,6 +315,82 @@ struct FArchetypeGroup_NewArchetype : FEntityTestBase
 };
 IMPLEMENT_AI_INSTANT_TEST(FArchetypeGroup_NewArchetype, "System.Mass.Archetype.Group.NewArchetype");
 
+struct FArchetypeGroup_NonGroupChange : FEntityTestBase
+{
+	virtual bool InstantTest() override
+	{
+		uint32 ArbitraryNumber = 1677;
+
+		FArchetypeGroupType GroupType = EntityManager->FinOrAddArchetypeGroupType(TEXT("TestGroup"));
+		const FArchetypeGroupHandle TargetGroupHandle(GroupType, ArbitraryNumber);
+
+		const FMassEntityHandle Entity = EntityManager->CreateEntity(IntsArchetype);
+		
+		EntityManager->BatchGroupEntities(TargetGroupHandle, MakeArrayView(&Entity, 1));
+		FMassArchetypeHandle ArchetypeGroup = EntityManager->GetArchetypeForEntity(Entity);
+
+		EntityManager->AddTagToEntity(Entity, FTestTag_A::StaticStruct());
+
+		FMassArchetypeHandle ArchetypeTagGroup = EntityManager->GetArchetypeForEntity(Entity);
+		AITEST_NOT_EQUAL("The archetype before and after the move", ArchetypeGroup, ArchetypeTagGroup);
+
+		const FArchetypeGroupHandle FinalGroupHandle = EntityManager->GetGroupForEntity(Entity, GroupType);
+		AITEST_EQUAL("Entity's group handle before and after tagging", TargetGroupHandle, FinalGroupHandle);
+
+		return true;
+	}
+};
+IMPLEMENT_AI_INSTANT_TEST(FArchetypeGroup_NonGroupChange, "System.Mass.Archetype.Group.AddTag");
+
+/** Ensuring that grouping an entity moves it to a different archetype*/
+struct FArchetypeGroup_Remove : FEntityTestBase
+{
+	virtual bool InstantTest() override
+	{
+		constexpr uint32 ArbitraryNumber = 1677;
+		const FArchetypeGroupType GroupTypeA = EntityManager->FinOrAddArchetypeGroupType(TEXT("TestGroupA"));
+		const FArchetypeGroupType GroupTypeB = EntityManager->FinOrAddArchetypeGroupType(TEXT("TestGroupB"));
+		const FMassArchetypeHandle OriginalArchetype = IntsArchetype;
+
+		TArray<FMassEntityHandle> Entities;
+		EntityManager->BatchCreateEntities(OriginalArchetype, 2, Entities);
+
+		EntityManager->BatchGroupEntities(FArchetypeGroupHandle(GroupTypeA, ArbitraryNumber), Entities);
+		const FMassArchetypeHandle ArchetypeGroupA = EntityManager->GetArchetypeForEntity(Entities[0]);
+		EntityManager->BatchGroupEntities(FArchetypeGroupHandle(GroupTypeB, ArbitraryNumber), Entities);
+		const FMassArchetypeHandle ArchetypeGroupAB = EntityManager->GetArchetypeForEntity(Entities[0]);
+
+		{
+			EntityManager->RemoveEntityFromGroupType(Entities[0], GroupTypeA);
+			const FMassArchetypeHandle ArchetypeGroupNoA = EntityManager->GetArchetypeForEntity(Entities[0]);
+
+			const FArchetypeGroups& GroupsNoA = EntityManager->GetGroupsForArchetype(ArchetypeGroupNoA);
+			AITEST_TRUE("After removing group A from AB entity it's still in group B", GroupsNoA.ContainsType(GroupTypeB));
+			AITEST_FALSE("After removing group A from AB entity it's still in group A", GroupsNoA.ContainsType(GroupTypeA));
+
+			EntityManager->RemoveEntityFromGroupType(Entities[0], GroupTypeB);
+			const FMassArchetypeHandle ArchetypeNoGroups = EntityManager->GetArchetypeForEntity(Entities[0]);
+			AITEST_EQUAL("The archetype after removing the entity from both groups; and the original archetype, scenario 1", ArchetypeNoGroups, OriginalArchetype);
+		}
+		{
+			// let's do the same in the other order, first remove B then A
+			EntityManager->RemoveEntityFromGroupType(Entities[1], GroupTypeB);
+			const FMassArchetypeHandle ArchetypeGroupNoB = EntityManager->GetArchetypeForEntity(Entities[1]);
+
+			const FArchetypeGroups& GroupsNoB = EntityManager->GetGroupsForArchetype(ArchetypeGroupNoB);
+			AITEST_FALSE("After removing group B from AB entity it's still in group B", GroupsNoB.ContainsType(GroupTypeB));
+			AITEST_TRUE("After removing group B from AB entity it's still in group A", GroupsNoB.ContainsType(GroupTypeA));
+
+			EntityManager->RemoveEntityFromGroupType(Entities[1], GroupTypeA);
+			const FMassArchetypeHandle ArchetypeNoGroups = EntityManager->GetArchetypeForEntity(Entities[1]);
+			AITEST_EQUAL("The archetype after removing the entity from both groups; and the original archetype, scenario 2", ArchetypeNoGroups, OriginalArchetype);
+		}
+
+		return true;
+	}
+};
+IMPLEMENT_AI_INSTANT_TEST(FArchetypeGroup_Remove, "System.Mass.Archetype.Group.Remove");
+
 } // UE::Mass::Test
 
 UE_ENABLE_OPTIMIZATION_SHIP
diff --git a/Engine/Source/Runtime/MassEntity/Internal/MassArchetypeData.h b/Engine/Source/Runtime/MassEntity/Internal/MassArchetypeData.h
--- a/Engine/Source/Runtime/MassEntity/Internal/MassArchetypeData.h
+++ b/Engine/Source/Runtime/MassEntity/Internal/MassArchetypeData.h
@@ -220,6 +220,7 @@ public:
 
 	const UE::Mass::FArchetypeGroups& GetGroups() const;
 	bool IsInGroup(const UE::Mass::FArchetypeGroupHandle GroupHandle) const;
+	bool IsInGroupOfType(const UE::Mass::FArchetypeGroupType GroupType) const;
 
 	/** Method to iterate on all the fragment types */
 	void ForEachFragmentType(TFunction< void(const UScriptStruct* /*FragmentType*/)> Function) const;
@@ -235,7 +236,7 @@ public:
 	 * replacing original tags of BaseArchetype
 	 */
 	void InitializeWithSimilar(const FMassEntityManager& EntityManager, const FMassArchetypeData& BaseArchetype
-		, FMassArchetypeCompositionDescriptor&& NewComposition, UE::Mass::FArchetypeGroups&& InGroups, const uint32 ArchetypeDataVersion);
+		, FMassArchetypeCompositionDescriptor&& NewComposition, const UE::Mass::FArchetypeGroups& InGroups, const uint32 ArchetypeDataVersion);
 
 	void AddEntity(FMassEntityHandle Entity, const FMassArchetypeSharedFragmentValues& InSharedFragmentValues);
 	void RemoveEntity(FMassEntityHandle Entity);
@@ -470,3 +471,8 @@ inline bool FMassArchetypeData::IsInGroup(const UE::Mass::FArchetypeGroupHandle
 	}
 	return false;
 }
+
+inline bool FMassArchetypeData::IsInGroupOfType(const UE::Mass::FArchetypeGroupType GroupType) const
+{
+	return Groups.ContainsType(GroupType);
+}
diff --git a/Engine/Source/Runtime/MassEntity/Private/MassArchetypeData.cpp b/Engine/Source/Runtime/MassEntity/Private/MassArchetypeData.cpp
--- a/Engine/Source/Runtime/MassEntity/Private/MassArchetypeData.cpp
+++ b/Engine/Source/Runtime/MassEntity/Private/MassArchetypeData.cpp
@@ -92,7 +92,7 @@ void FMassArchetypeData::Initialize(const FMassEntityManager& EntityManager, con
 }
 
 void FMassArchetypeData::InitializeWithSimilar(const FMassEntityManager& EntityManager, const FMassArchetypeData& BaseArchetype
-	, FMassArchetypeCompositionDescriptor&& NewComposition, UE::Mass::FArchetypeGroups&& InGroups, const uint32 ArchetypeDataVersion)
+	, FMassArchetypeCompositionDescriptor&& NewComposition, const UE::Mass::FArchetypeGroups& InGroups, const uint32 ArchetypeDataVersion)
 {
 	checkf(IsInitialized() == false, TEXT("Trying to InitializeWithSimilar but this archetype has already been initialized"));
 
@@ -113,7 +113,7 @@ void FMassArchetypeData::InitializeWithSimilar(const FMassEntityManager& EntityM
 	}
 	ChunkFragmentsTemplate = BaseArchetype.ChunkFragmentsTemplate;
 
-	Groups.Set(Forward<UE::Mass::FArchetypeGroups>(InGroups));
+	Groups = InGroups;
 
 	EntityListOffsetWithinChunk = 0;
 }
diff --git a/Engine/Source/Runtime/MassEntity/Private/MassArchetypeGroup.cpp b/Engine/Source/Runtime/MassEntity/Private/MassArchetypeGroup.cpp
--- a/Engine/Source/Runtime/MassEntity/Private/MassArchetypeGroup.cpp
+++ b/Engine/Source/Runtime/MassEntity/Private/MassArchetypeGroup.cpp
@@ -13,9 +13,16 @@ namespace UE::Mass
 		return IDContainer == OtherGroups.IDContainer;
 	}
 
-	void FArchetypeGroups::Set(FArchetypeGroups&& InGroups)
+	FArchetypeGroups& FArchetypeGroups::operator=(FArchetypeGroups&& InGroups)
 	{
 		IDContainer = MoveTemp(InGroups.IDContainer);
+		return *this;
+	}
+
+	FArchetypeGroups& FArchetypeGroups::operator=(const FArchetypeGroups& InGroups)
+	{
+		IDContainer = InGroups.IDContainer;
+		return *this;
 	}
 
 	void FArchetypeGroups::Add(FArchetypeGroupHandle GroupHandle)
@@ -31,6 +38,68 @@ namespace UE::Mass
 		}
 	}
 
+	FArchetypeGroups FArchetypeGroups::Add(FArchetypeGroupHandle GroupHandle) const
+	{
+		FArchetypeGroups Copy = *this;
+		const int32 GroupTypeIndex = static_cast<int32>(GroupHandle.GetGroupType());
+		if (IDContainer.IsValidIndex(GroupTypeIndex))
+		{
+			Copy.IDContainer[GroupTypeIndex] = GroupHandle.GetGroupID();
+		}
+		else
+		{
+			Copy.IDContainer.EmplaceAt(GroupTypeIndex, GroupHandle.GetGroupID());
+		}
+
+		return MoveTemp(Copy);
+	}
+
+	void FArchetypeGroups::Remove(FArchetypeGroupType GroupType)
+	{
+		const int32 GroupTypeIndex = static_cast<int32>(GroupType);
+		if (IDContainer.IsValidIndex(GroupTypeIndex))
+		{
+			const bool IsLastElement = ((IDContainer.GetMaxIndex() - 1) == GroupTypeIndex);
+			IDContainer.RemoveAtUninitialized(GroupTypeIndex);
+			if (IsLastElement)
+			{
+				Shrink();
+			}
+		}
+	}
+
+	FArchetypeGroups FArchetypeGroups::Remove(FArchetypeGroupType GroupType) const
+	{
+		const int32 GroupTypeIndex = static_cast<int32>(GroupType);
+		if (IDContainer.IsValidIndex(GroupTypeIndex))
+		{
+			FArchetypeGroups Copy = *this;
+			Copy.IDContainer.RemoveAtUninitialized(GroupTypeIndex);
+
+			// if we removed the last element we need to shrink the container too.
+			if (const bool IsLastElement = ((IDContainer.GetMaxIndex() - 1) == GroupTypeIndex))
+			{
+				Copy.Shrink();
+			}
+
+			return MoveTemp(Copy);
+		}
+		return *this;
+	}
+
+	void FArchetypeGroups::Shrink()
+	{
+		IDContainer.Shrink();
+	}
+
+	bool FArchetypeGroups::IsShrunk() const
+	{
+		const int32 MaxIndex = IDContainer.GetMaxIndex();
+		// if the IDContainer has been shrunk, or never needed to be shrunk, the last element in the available
+		// container is valid.
+		return MaxIndex == 0 || IDContainer.IsValidIndex(MaxIndex - 1) ;
+	}
+
 	uint32 GetTypeHash(const FArchetypeGroups& Instance)
 	{
 		uint32 Hash = 0;
diff --git a/Engine/Source/Runtime/MassEntity/Private/MassEntityManager.cpp b/Engine/Source/Runtime/MassEntity/Private/MassEntityManager.cpp
--- a/Engine/Source/Runtime/MassEntity/Private/MassEntityManager.cpp
+++ b/Engine/Source/Runtime/MassEntity/Private/MassEntityManager.cpp
@@ -483,7 +483,8 @@ FMassArchetypeHandle FMassEntityManager::InternalCreateSimilarArchetype(const TS
 		, SourceArchetypeRef.GetChunkFragmentBitSet()
 		, SourceArchetypeRef.GetSharedFragmentBitSet()
 		, SourceArchetypeRef.GetConstSharedFragmentBitSet());
-	return InternalCreateSimilarArchetype(SourceArchetypeRef, MoveTemp(NewComposition), {});
+
+	return InternalCreateSimilarArchetype(SourceArchetypeRef, MoveTemp(NewComposition), SourceArchetypeRef.GetGroups());
 }
 
 FMassArchetypeHandle FMassEntityManager::InternalCreateSimilarArchetype(const TSharedPtr<FMassArchetypeData>& SourceArchetype, const FMassFragmentBitSet& OverrideFragments)
@@ -495,19 +496,24 @@ FMassArchetypeHandle FMassEntityManager::InternalCreateSimilarArchetype(const TS
 		, SourceArchetypeRef.GetChunkFragmentBitSet()
 		, SourceArchetypeRef.GetSharedFragmentBitSet()
 		, SourceArchetypeRef.GetConstSharedFragmentBitSet());
-	return InternalCreateSimilarArchetype(SourceArchetypeRef, MoveTemp(NewComposition), {});
+
+	return InternalCreateSimilarArchetype(SourceArchetypeRef, MoveTemp(NewComposition), SourceArchetypeRef.GetGroups());
 }
 
-FMassArchetypeHandle FMassEntityManager::InternalCreateSimilarArchetype(const TSharedPtr<FMassArchetypeData>& SourceArchetype, UE::Mass::FArchetypeGroups&& GroupsOverride)
+FMassArchetypeHandle FMassEntityManager::InternalCreateSimilarArchetype(const TSharedPtr<FMassArchetypeData>& SourceArchetype, const UE::Mass::FArchetypeGroups& GroupsOverride)
 {
 	checkSlow(SourceArchetype.IsValid());
 	const FMassArchetypeData& SourceArchetypeRef = *SourceArchetype.Get();
 	FMassArchetypeCompositionDescriptor NewComposition = SourceArchetype->GetCompositionDescriptor();
-	return InternalCreateSimilarArchetype(SourceArchetypeRef, MoveTemp(NewComposition), Forward<UE::Mass::FArchetypeGroups>(GroupsOverride));
+	return InternalCreateSimilarArchetype(SourceArchetypeRef, MoveTemp(NewComposition), GroupsOverride);
 }
 
-FMassArchetypeHandle FMassEntityManager::InternalCreateSimilarArchetype(const FMassArchetypeData& SourceArchetypeRef, FMassArchetypeCompositionDescriptor&& NewComposition, UE::Mass::FArchetypeGroups&& Groups)
+FMassArchetypeHandle FMassEntityManager::InternalCreateSimilarArchetype(const FMassArchetypeData& SourceArchetypeRef, FMassArchetypeCompositionDescriptor&& NewComposition, const UE::Mass::FArchetypeGroups& Groups)
 {
+	// we require Groups to be already shrunk. Shrinking is required to remove any trailing, invalid group IDs that would
+	// be there if IDs were added and removed to this specific Groups container instance
+	checkf(Groups.IsShrunk(), TEXT("A group container with invalid trailing IDs has been passed to archetype creation - this is not expected and will cause issues. Make sure to Shrink your Groups before passing to %hs"), __FUNCTION__);
+
 	const uint32 TypeHash = HashCombine(NewComposition.CalculateHash(), GetTypeHash(Groups));
 
 	TArray<TSharedPtr<FMassArchetypeData>>& HashRow = FragmentHashToArchetypeMap.FindOrAdd(TypeHash);
@@ -529,7 +535,7 @@ FMassArchetypeHandle FMassEntityManager::InternalCreateSimilarArchetype(const FM
 
 		// Create a new archetype
 		FMassArchetypeData* NewArchetype = new FMassArchetypeData(FMassArchetypeCreationParams(SourceArchetypeRef));
-		NewArchetype->InitializeWithSimilar(*this, SourceArchetypeRef, MoveTemp(NewComposition), Forward<UE::Mass::FArchetypeGroups>(Groups), ArchetypeDataVersion);
+		NewArchetype->InitializeWithSimilar(*this, SourceArchetypeRef, MoveTemp(NewComposition), Groups, ArchetypeDataVersion);
 		NewArchetype->CopyDebugNamesFrom(SourceArchetypeRef);
 
 		ArchetypeDataPtr = HashRow.Add_GetRef(MakeShareable(NewArchetype));
@@ -986,6 +992,35 @@ void FMassEntityManager::BatchGroupEntities(const UE::Mass::FArchetypeGroupHandl
 	BatchGroupEntities(GroupHandle, Collections);
 }
 
+void FMassEntityManager::RemoveEntityFromGroupType(FMassEntityHandle EntityHandle, UE::Mass::FArchetypeGroupType GroupType)
+{
+	const FMassArchetypeHandle CurrentArchetypeHandle = GetArchetypeForEntity(EntityHandle);
+	if (FMassArchetypeData* CurrentArchetype = CurrentArchetypeHandle.DataPtr.Get())
+	{
+		if (CurrentArchetype->IsInGroupOfType(GroupType))
+		{
+			/*UE::Mass::FArchetypeGroups NewGroups = CurrentArchetype->GetGroups();
+			NewGroups.Remove(GroupType);
+			const FMassArchetypeHandle NewArchetypeHandle = InternalCreateSimilarArchetype(CurrentArchetypeHandle.DataPtr, MoveTemp(NewGroups));*/
+			const FMassArchetypeHandle NewArchetypeHandle = InternalCreateSimilarArchetype(CurrentArchetypeHandle.DataPtr, CurrentArchetype->GetGroups().Remove(GroupType));
+
+			CurrentArchetype->MoveEntityToAnotherArchetype(EntityHandle, *NewArchetypeHandle.DataPtr.Get());
+			
+			GetEntityStorageInterface().SetArchetypeFromShared(EntityHandle.Index, NewArchetypeHandle.DataPtr);
+		}
+	}
+}
+
+UE::Mass::FArchetypeGroupHandle FMassEntityManager::GetGroupForEntity(FMassEntityHandle EntityHandle, UE::Mass::FArchetypeGroupType GroupType) const
+{
+	if (FMassArchetypeData* CurrentArchetype = GetArchetypeForEntity(EntityHandle).DataPtr.Get())
+	{
+		return UE::Mass::FArchetypeGroupHandle(GroupType, CurrentArchetype->GetGroups().GetID(GroupType));
+	}
+
+	return UE::Mass::FArchetypeGroupHandle();
+}
+
 UE::Mass::FArchetypeGroupType FMassEntityManager::FinOrAddArchetypeGroupType(const FName GroupName)
 {
 	const int32* FoundGroupIndex = GroupNameToTypeIndex.Find(GroupName);
diff --git a/Engine/Source/Runtime/MassEntity/Public/MassArchetypeGroup.h b/Engine/Source/Runtime/MassEntity/Public/MassArchetypeGroup.h
--- a/Engine/Source/Runtime/MassEntity/Public/MassArchetypeGroup.h
+++ b/Engine/Source/Runtime/MassEntity/Public/MassArchetypeGroup.h
@@ -13,11 +13,15 @@ namespace UE::Mass
 	//-----------------------------------------------------------------------------
 	struct FArchetypeGroupType
 	{
-		explicit FArchetypeGroupType(int32 Value = INDEX_NONE)
+		constexpr static uint32 InvalidArchetypeGroupTypeIdentifier = static_cast<uint32>(INDEX_NONE);
+
+		explicit FArchetypeGroupType(const uint32 Value = InvalidArchetypeGroupTypeIdentifier)
 			: Identifier(Value)
 		{
 		}
 
+		FArchetypeGroupType(const FArchetypeGroupType& Source) = default;
+
 		friend uint32 GetTypeHash(const FArchetypeGroupType& Instance)
 		{
 			return GetTypeHash(Instance.Identifier);
@@ -35,27 +39,33 @@ namespace UE::Mass
 
 		explicit operator int32() const
 		{
-			return Identifier;
+			return static_cast<int32>(Identifier);
 		}
 
 		bool IsValid() const
 		{
-			return Identifier != INDEX_NONE;
+			return Identifier != InvalidArchetypeGroupTypeIdentifier;
 		}
 
 	private:
-		const int32 Identifier;
+		uint32 Identifier;
 	};
 
 
 	struct FArchetypeGroupID
 	{
 		constexpr static uint32 InvalidArchetypeGroupID = static_cast<uint32>(INDEX_NONE);
+		constexpr static uint32 FirstGroupID = 0;
 
-		FArchetypeGroupID(const uint32 InID = InvalidArchetypeGroupID)
+		FArchetypeGroupID() = default;
+		FArchetypeGroupID(const uint32 InID)
 			: ID(InID)
 		{
-			
+		}
+
+		bool operator==(const FArchetypeGroupID Other) const
+		{
+			return ID == Other.ID;
 		}
 
 		bool IsValid() const
@@ -68,8 +78,18 @@ namespace UE::Mass
 			return static_cast<int32>(ID);
 		}
 
+		static FArchetypeGroupID First()
+		{
+			return FArchetypeGroupID(FirstGroupID);
+		}
+
+		FArchetypeGroupID Next() const
+		{
+			return FArchetypeGroupID(ID + 1);
+		}
+
 	private:
-		uint32 ID;
+		uint32 ID = InvalidArchetypeGroupID;
 	};
 
 	//-----------------------------------------------------------------------------
@@ -77,11 +97,23 @@ namespace UE::Mass
 	//-----------------------------------------------------------------------------
 	struct FArchetypeGroupHandle
 	{
-		FArchetypeGroupHandle(const FArchetypeGroupType InGroupType, const FArchetypeGroupID InGroupID)
+		explicit FArchetypeGroupHandle(const FArchetypeGroupType InGroupType, const FArchetypeGroupID InGroupID)
 			: GroupType(InGroupType), GroupID(InGroupID)
 		{
 		}
 
+		FArchetypeGroupHandle() = default;
+
+		bool operator==(const FArchetypeGroupHandle Other) const
+		{
+			return GroupType == Other.GroupType && GroupID == Other.GroupID;
+		}
+
+		bool operator!=(const FArchetypeGroupHandle Other) const
+		{
+			return !(*this == Other);
+		}
+
 		friend uint32 GetTypeHash(const FArchetypeGroupHandle& Instance)
 		{
 			const uint64 CombinedHandle = static_cast<uint64>(static_cast<int32>(Instance.GroupType)) << 32 | Instance.GroupID;
@@ -117,8 +149,8 @@ namespace UE::Mass
 		}
 
 	private:
-		const FArchetypeGroupType GroupType;
-		const FArchetypeGroupID GroupID;
+		FArchetypeGroupType GroupType;
+		FArchetypeGroupID GroupID;
 	};
 
 	//-----------------------------------------------------------------------------
@@ -126,15 +158,53 @@ namespace UE::Mass
 	//-----------------------------------------------------------------------------
 	struct FArchetypeGroups
 	{
+		FArchetypeGroups() = default;
+		FArchetypeGroups(const FArchetypeGroups& InGroups) = default;
+		FArchetypeGroups(FArchetypeGroups&& InGroups) = default;
+
 		bool operator==(const FArchetypeGroups& OtherGroups) const;
-		void Set(UE::Mass::FArchetypeGroups&& InGroups);
+		FArchetypeGroups& operator=(FArchetypeGroups&& InGroups);
+		FArchetypeGroups& operator=(const FArchetypeGroups& InGroups);
+
+		/**
+		 * Adds or updates the given (GroupType, GroupID) combination to IDContainer
+		 */
 		void Add(FArchetypeGroupHandle GroupHandle);
 
+		/**
+		 * Adds or updates the given (GroupType, GroupID) combination to IDContainer
+		 * @return a copy of this FArchetypeGroups container with GroupHandle added to the ID container
+		 * @note using [[nodiscard]] to avoid accidental calls on const instances that would not produce any effects
+		 */
+		[[nodiscard]] FArchetypeGroups Add(FArchetypeGroupHandle GroupHandle) const;
+
+		/**
+		 * Removes the stored GroupID associated with the given GroupType.
+		 * If the given group type is not stored in IDContainer the request is ignored.
+		 */
+		void Remove(FArchetypeGroupType GroupType);
+
+		/**
+		 * Removes the stored GroupID associated with the given GroupType.
+		 * If the given group type is not stored in IDContainer the request is ignored.
+		 * @return a copy of this FArchetypeGroups container with GroupType removed.
+		 * @note using [[nodiscard]] to avoid accidental calls on const instances that would not produce any effects
+		 */
+		[[nodiscard]] FArchetypeGroups Remove(FArchetypeGroupType GroupType) const;
+
+		void Shrink();
+		bool IsShrunk() const;
+
 		FArchetypeGroupID GetID(const FArchetypeGroupType GroupType) const
 		{
 			return IDContainer.IsValidIndex(static_cast<int32>(GroupType)) ? IDContainer[static_cast<int32>(GroupType)] : FArchetypeGroupID();
 		}
 
+		bool ContainsType(const FArchetypeGroupType GroupType) const
+		{
+			return IDContainer.IsValidIndex(static_cast<int32>(GroupType));
+		}
+
 		friend uint32 GetTypeHash(const FArchetypeGroups& Instance);
 
 	protected:
diff --git a/Engine/Source/Runtime/MassEntity/Public/MassEntityManager.h b/Engine/Source/Runtime/MassEntity/Public/MassEntityManager.h
--- a/Engine/Source/Runtime/MassEntity/Public/MassEntityManager.h
+++ b/Engine/Source/Runtime/MassEntity/Public/MassEntityManager.h
@@ -420,6 +420,17 @@ public:
 	 */
 	bool RemoveSharedFragmentFromEntity(const FMassEntityHandle Entity, const UScriptStruct& SharedFragmentType);
 
+	/**
+	 * Removes EntityHandle from any-and-all groups of given type - i.e. the entity will be moved to an archetype
+	 * not in any of the groups of the given type.
+	 */
+	void RemoveEntityFromGroupType(FMassEntityHandle EntityHandle, UE::Mass::FArchetypeGroupType GroupType);
+
+	/**
+	 * @return the group handle of the specific group of type GroupType that the entity belongs to
+	 */
+	UE::Mass::FArchetypeGroupHandle GetGroupForEntity(FMassEntityHandle EntityHandle, UE::Mass::FArchetypeGroupType GroupType) const;
+
 	/**
 	 * Reserves Count number of entities and appends them to InOutEntities
 	 * @return a view into InOutEntities containing only the freshly reserved entities
@@ -1049,8 +1060,8 @@ protected:
 	 */
 	FMassArchetypeHandle InternalCreateSimilarArchetype(const TSharedPtr<FMassArchetypeData>& SourceArchetype, const FMassTagBitSet& OverrideTags);
 	FMassArchetypeHandle InternalCreateSimilarArchetype(const TSharedPtr<FMassArchetypeData>& SourceArchetype, const FMassFragmentBitSet& OverrideFragments);
-	FMassArchetypeHandle InternalCreateSimilarArchetype(const TSharedPtr<FMassArchetypeData>& SourceArchetype, UE::Mass::FArchetypeGroups&& GroupsOverride);
-	FMassArchetypeHandle InternalCreateSimilarArchetype(const FMassArchetypeData& SourceArchetypeRef, FMassArchetypeCompositionDescriptor&& NewComposition, UE::Mass::FArchetypeGroups&& GroupsOverride);
+	FMassArchetypeHandle InternalCreateSimilarArchetype(const TSharedPtr<FMassArchetypeData>& SourceArchetype, const UE::Mass::FArchetypeGroups& GroupsOverride);
+	FMassArchetypeHandle InternalCreateSimilarArchetype(const FMassArchetypeData& SourceArchetypeRef, FMassArchetypeCompositionDescriptor&& NewComposition, const UE::Mass::FArchetypeGroups& GroupsOverride);
 
 	void InternalAppendFragmentsAndTagsToArchetypeCompositionDescriptor(FMassArchetypeCompositionDescriptor& InOutComposition,
 		TConstArrayView<const UScriptStruct*> FragmentsAndTagsList) const;
@@ -1229,6 +1240,6 @@ private:
 	//-----------------------------------------------------------------------------
 	// DEPRECATED
 	//-----------------------------------------------------------------------------
-	UE_DEPRECATED(5.5, "This flavor of InternalCreateSimilarArchetype is deprecaded due to the introduction of archetype grouping. Use InternalCreateSimilarArchetype with a FArchetypeGroups parameter instead")
+	UE_DEPRECATED(5.6, "This flavor of InternalCreateSimilarArchetype is deprecated due to the introduction of archetype grouping. Use InternalCreateSimilarArchetype with a FArchetypeGroups parameter instead")
 	FMassArchetypeHandle InternalCreateSimilarArchetype(const FMassArchetypeData& SourceArchetypeRef, FMassArchetypeCompositionDescriptor&& NewComposition);
 };
