[Mass] EntityBuilder and Mass concepts

The CL introduces a new way of procedurally build entities, without any need for templates nor entity confirgs.
The CL also introduces Mass "concepts" that simplify

diff --git a/Engine/Source/Developer/MassEntityTestSuite/Private/MassEntityBuilderTest.cpp b/Engine/Source/Developer/MassEntityTestSuite/Private/MassEntityBuilderTest.cpp
new file mode 100644
--- /dev/null
+++ b/Engine/Source/Developer/MassEntityTestSuite/Private/MassEntityBuilderTest.cpp
@@ -0,0 +1,353 @@
+// Copyright Epic Games, Inc. All Rights Reserved.
+
+#include "AITestsCommon.h"
+#include "MassEntityTestTypes.h"
+#include "MassEntityBuilder.h"
+#include "MassEntityManager.h"
+
+
+#define LOCTEXT_NAMESPACE "MassTest"
+
+UE_DISABLE_OPTIMIZATION_SHIP
+
+/**
+ * Tests to be added:
+ * - test observers triggering as expected, i.e. respecting the construction context
+ * - entity grouping
+ */
+
+namespace UE::Mass::Test::EntityBuilder
+{
+	struct FSimpleBuild : FEntityTestBase
+	{
+		virtual bool InstantTest() override
+		{
+			int32 EntitiesCreated = EntityManager->DebugGetArchetypeEntitiesCount(IntsArchetype);
+			int32 SomeCounter = 1;
+			int32 EntitiesCreatedThisStep = 0;
+
+			{
+				FEntityBuilder EntityBuilder(*EntityManager.Get());
+				EntityBuilder.Add_GetRef<FTestFragment_Int>(SomeCounter++);
+				EntityBuilder.Commit();
+				EntitiesCreatedThisStep = 1;
+			}
+
+			AITEST_EQUAL("Number of entities created with basic use", EntityManager->DebugGetArchetypeEntitiesCount(IntsArchetype) - EntitiesCreated, EntitiesCreatedThisStep);
+			EntitiesCreated = EntityManager->DebugGetArchetypeEntitiesCount(IntsArchetype);
+
+			{
+				FEntityBuilder EntityBuilder(*EntityManager.Get());
+				ON_SCOPE_EXIT
+				{
+					EntityBuilder.Commit();
+				};
+
+				EntityBuilder.Add_GetRef<FTestFragment_Int>(SomeCounter++);
+			}
+			EntitiesCreatedThisStep = 1;
+
+			AITEST_EQUAL("Number of entities created with ON_SCOPE_EXIT", EntityManager->DebugGetArchetypeEntitiesCount(IntsArchetype) - EntitiesCreated, EntitiesCreatedThisStep);
+			EntitiesCreated = EntityManager->DebugGetArchetypeEntitiesCount(IntsArchetype);
+
+			{
+				FScopedEntityBuilder EntityBuilder(*EntityManager.Get());
+				EntityBuilder.Add_GetRef<FTestFragment_Int>(SomeCounter++);
+			}
+			EntitiesCreatedThisStep = 1;
+
+			AITEST_EQUAL("Number of entities created with scoped builder", EntityManager->DebugGetArchetypeEntitiesCount(IntsArchetype) - EntitiesCreated, EntitiesCreatedThisStep);
+			EntitiesCreated = EntityManager->DebugGetArchetypeEntitiesCount(IntsArchetype);
+
+			{
+				FScopedEntityBuilder EntityBuilder(*EntityManager.Get());
+				EntityBuilder.Add_GetRef<FTestFragment_Int>(SomeCounter++);
+
+				FEntityBuilder EntityBuilder2 = EntityBuilder;
+				EntityBuilder2.Commit();
+			}
+			EntitiesCreatedThisStep = 2;
+
+			AITEST_EQUAL("Number of entities created with with a scoped builder and its regular copy", EntityManager->DebugGetArchetypeEntitiesCount(IntsArchetype) - EntitiesCreated, EntitiesCreatedThisStep);
+			EntitiesCreated = EntityManager->DebugGetArchetypeEntitiesCount(IntsArchetype);
+
+			{
+				FEntityBuilder EntityBuilder(*EntityManager.Get());
+				EntityBuilder.Add_GetRef<FTestFragment_Int>(SomeCounter++);
+				EntityBuilder.Commit();
+
+				FEntityBuilder EntityBuilder2 = EntityBuilder;
+			}
+			EntitiesCreatedThisStep = 1;
+
+			AITEST_EQUAL("Number of entities created with Commit and an abandoned copy of a builder", EntityManager->DebugGetArchetypeEntitiesCount(IntsArchetype) - EntitiesCreated, EntitiesCreatedThisStep);
+			EntitiesCreated = EntityManager->DebugGetArchetypeEntitiesCount(IntsArchetype);
+
+			{
+				FEntityBuilder EntityBuilder(*EntityManager.Get());
+				EntityBuilder.Add_GetRef<FTestFragment_Int>(SomeCounter++);
+
+				FEntityBuilder EntityBuilder2 = EntityBuilder;
+				EntityBuilder2.Commit();
+
+				EntityBuilder.Reset();
+			}
+			EntitiesCreatedThisStep = 1;
+
+			AITEST_EQUAL("Number of entities created with committed copy and abandoned original builder", EntityManager->DebugGetArchetypeEntitiesCount(IntsArchetype) - EntitiesCreated, EntitiesCreatedThisStep);
+			EntitiesCreated = EntityManager->DebugGetArchetypeEntitiesCount(IntsArchetype);
+
+			return true;
+		}
+	};
+	IMPLEMENT_AI_INSTANT_TEST(FSimpleBuild, "System.Mass.EntityBuilder.SimpleBuild");
+
+	struct FAbort: FEntityTestBase
+	{
+		virtual bool InstantTest() override
+		{
+			{
+				FEntityBuilder Builder(EntityManager.ToSharedRef());
+				Builder.Add<FTestFragment_Int>();
+				FMassEntityHandle ReservedEntityHandle = Builder.GetEntityHandle();
+				{
+					const bool bValid = EntityManager->IsEntityValid(ReservedEntityHandle);
+					AITEST_TRUE("Before committing the entity handle is reserved", bValid);
+					const bool bIsBuilt = EntityManager->IsEntityActive(ReservedEntityHandle);
+					AITEST_FALSE("Before committing the entity is already created", bIsBuilt);
+				}
+				Builder.Reset();
+				{
+					const bool bValid = EntityManager->IsEntityValid(ReservedEntityHandle);
+					AITEST_FALSE("After resetting the entity handle is still valid", bValid);
+				}
+			}
+			FMassEntityHandle AbandonedEntityHandle;
+			{
+				FEntityBuilder Builder(EntityManager.ToSharedRef());
+				Builder.Add<FTestFragment_Int>();
+				AbandonedEntityHandle = Builder.GetEntityHandle();
+			}
+			{
+				const bool bValid = EntityManager->IsEntityValid(AbandonedEntityHandle);
+				AITEST_FALSE("After builder's destruction without committing the entity handle is still valid", bValid);
+			}
+
+			return true;
+		}
+	};
+	IMPLEMENT_AI_INSTANT_TEST(FAbort, "System.Mass.EntityBuilder.Abort");
+
+	struct FOneliner: FEntityTestBase
+	{
+		virtual bool InstantTest() override
+		{
+			int32 TotalCountCreated = 0;
+			{
+				FMassEntityHandle CreatedEntity = EntityManager->MakeEntityBuilder().Add<FTestFragment_Int>().Commit();
+				++TotalCountCreated;
+
+				const bool bIsBuilt = EntityManager->IsEntityActive(CreatedEntity);
+				AITEST_TRUE("The entity has been created", bIsBuilt);
+#if WITH_MASSENTITY_DEBUG
+				AITEST_TRUE("Only a single entity has been created", EntityManager->DebugGetArchetypeEntitiesCount(IntsArchetype) == TotalCountCreated);
+#endif // WITH_MASSENTITY_DEBUG
+			}
+			{
+				FEntityBuilder EntityBuilder = EntityManager->MakeEntityBuilder().Add<FTestFragment_Int>();
+				EntityBuilder.Commit();
+				++TotalCountCreated;
+			}
+#if WITH_MASSENTITY_DEBUG
+			AITEST_TRUE("The number of entities created matches expectations", EntityManager->DebugGetArchetypeEntitiesCount(IntsArchetype) == TotalCountCreated);
+#endif // WITH_MASSENTITY_DEBUG
+			{
+				// we're not committing so this builder won't create an entity.
+				EntityManager->MakeEntityBuilder().Add<FTestFragment_Int>();
+
+				// similarly here, even reserving the entity won't result in building that entity without manual `Commit` call.
+				FMassEntityHandle ReservedEntity = EntityManager->MakeEntityBuilder().Add<FTestFragment_Int>().GetEntityHandle();
+			}
+#if WITH_MASSENTITY_DEBUG
+			AITEST_EQUAL("The number of entities created after not committing builders", EntityManager->DebugGetArchetypeEntitiesCount(IntsArchetype), TotalCountCreated);
+#endif // WITH_MASSENTITY_DEBUG
+
+			return true;
+		}
+	};
+	IMPLEMENT_AI_INSTANT_TEST(FOneliner, "System.Mass.EntityBuilder.OneLiner");
+
+	struct FCopy : FEntityTestBase
+	{
+		virtual bool InstantTest() override
+		{
+			const int32 ValueA = 1;
+			FEntityBuilder BuilderA(*EntityManager.Get());
+			BuilderA.Add<FTestFragment_Int>(ValueA);
+
+			const int32 ValueB = ValueA + 1;
+			FEntityBuilder BuilderB = EntityManager->MakeEntityBuilder();
+			BuilderB.Add<FTestFragment_Int>(ValueB);
+
+			// a different way of setting the value
+			FEntityBuilder BuilderC = BuilderA;
+			BuilderC.GetOrCreate<FTestFragment_Int>().Value = ValueB;
+
+			BuilderA.Commit();
+			BuilderB.Commit();
+			BuilderC.Commit();
+
+			FTestFragment_Int* FragmentA = EntityManager->GetFragmentDataPtr<FTestFragment_Int>(BuilderA.GetEntityHandle());
+			AITEST_NOT_NULL("The original entity has the expected fragment", FragmentA);
+			AITEST_EQUAL("The value of the original entity's fragment matches expectations", FragmentA->Value, ValueA);
+
+			FTestFragment_Int* FragmentB = EntityManager->GetFragmentDataPtr<FTestFragment_Int>(BuilderB.GetEntityHandle());
+			AITEST_NOT_NULL("The copied entity has the expected fragment", FragmentB);
+			AITEST_EQUAL("The value of the copied entity's fragment matches expectations", FragmentB->Value, ValueB);
+
+			FTestFragment_Int* FragmentC = EntityManager->GetFragmentDataPtr<FTestFragment_Int>(BuilderC.GetEntityHandle());
+			AITEST_NOT_NULL("The other copied entity has the expected fragment", FragmentC);
+			AITEST_EQUAL("The value of the other copied entity's fragment matches expectations", FragmentC->Value, FragmentB->Value);
+
+			return true;
+		}
+	};
+	IMPLEMENT_AI_INSTANT_TEST(FCopy, "System.Mass.EntityBuilder.Copy");
+
+	struct FOverride : FEntityTestBase
+	{
+		virtual bool InstantTest() override
+		{
+			FEntityBuilder BuilderA(*EntityManager.Get());
+			BuilderA.Add<FTestFragment_Int>();
+
+			FEntityBuilder BuilderB(*EntityManager.Get());
+			BuilderB.Add<FTestFragment_Float>();
+
+			FEntityBuilder BuilderC = BuilderB;
+			
+			FMassEntityHandle EntityA = BuilderA.GetEntityHandle();
+			FMassEntityHandle EntityB = BuilderB.GetEntityHandle();
+			FMassEntityHandle EntityC = BuilderC.GetEntityHandle();
+
+			AITEST_NOT_EQUAL("Entities reserved by different builders, A|B", EntityA, EntityB);
+			AITEST_NOT_EQUAL("Entities reserved by different builders, A|C", EntityA, EntityC);
+			AITEST_NOT_EQUAL("Entities reserved by different builders, B|C", EntityB, EntityC);
+
+			// the following operation is expected to stomp the settings of the target builder, but not the entity
+			BuilderB = BuilderA;
+			FMassEntityHandle EntityB2 = BuilderB.GetEntityHandle();
+			AITEST_TRUE("The uncommitted (i.e. reserved) entity handle does not change with builder's config override", EntityB == EntityB2);
+
+			// overriding a committed builder results in creation of a new handle.
+			BuilderC.Commit();
+			BuilderC = BuilderA;
+			FMassEntityHandle EntityC2 = BuilderC.GetEntityHandle();
+			AITEST_TRUE("The committed entity handle differs from the new one", EntityC != EntityC2);
+
+			return true;
+		}
+	};
+	IMPLEMENT_AI_INSTANT_TEST(FOverride, "System.Mass.EntityBuilder.Override");
+
+	struct FPassOver : FEntityTestBase
+	{
+		virtual bool InstantTest() override
+		{
+			{
+				FEntityBuilder BuilderA(*EntityManager.Get());
+				BuilderA.Add<FTestFragment_Int>();
+				// forces handle reservation
+				FMassEntityHandle EntityA = BuilderA.GetEntityHandle();
+
+				FMassEntityHandle EntityB;
+				{
+					FEntityBuilder BuilderB(*EntityManager.Get());
+					BuilderB.Add<FTestFragment_Float>();
+					EntityB = BuilderB.GetEntityHandle();
+					BuilderA = MoveTemp(BuilderB);
+				}
+
+				// at this point EntityB should be valid while the original EntityA not
+				AITEST_TRUE("The original entity is invalid", EntityManager->IsEntityValid(EntityA) == false);
+				AITEST_TRUE("The passed-over entity entity is valid", EntityManager->IsEntityValid(EntityB));
+			}
+
+			{
+				FEntityBuilder BuilderA(*EntityManager.Get());
+				BuilderA.Add<FTestFragment_Int>();
+				// forces handle reservation
+				FMassEntityHandle EntityA = BuilderA.Commit();
+
+				FMassEntityHandle EntityB;
+				{
+					FEntityBuilder BuilderB(*EntityManager.Get());
+					BuilderB.Add<FTestFragment_Float>();
+					EntityB = BuilderB.GetEntityHandle();
+					BuilderA = MoveTemp(BuilderB);
+				}
+
+				// at this point EntityB should be valid while the original EntityA not
+				AITEST_TRUE("The original entity is valid, since it was committed", EntityManager->IsEntityValid(EntityA));
+				AITEST_TRUE("The original entity is active", EntityManager->IsEntityActive(EntityA));
+				AITEST_TRUE("The passed-over entity entity is valid", EntityManager->IsEntityValid(EntityB));
+			}
+
+			{
+				FEntityBuilder BuilderA(*EntityManager.Get());
+				BuilderA.Add<FTestFragment_Int>();
+				// forces handle reservation
+				FMassEntityHandle EntityA = BuilderA.GetEntityHandle();
+
+				FMassEntityHandle EntityB;
+				{
+					FEntityBuilder BuilderB(*EntityManager.Get());
+					BuilderB.Add<FTestFragment_Float>();
+					EntityB = BuilderB.Commit();
+					BuilderA = MoveTemp(BuilderB);
+				}
+
+				// at this point EntityB should be valid while the original EntityA not
+				AITEST_TRUE("The original entity is valid", EntityManager->IsEntityValid(EntityA));
+				AITEST_TRUE("The original entity is NOT active", EntityManager->IsEntityActive(EntityA) == false);
+				AITEST_TRUE("The secondary entity entity is valid", EntityManager->IsEntityValid(EntityB));
+				AITEST_TRUE("The secondary entity is active", EntityManager->IsEntityActive(EntityB));
+			}
+
+			{
+				FEntityBuilder BuilderA(*EntityManager.Get());
+				BuilderA.Add<FTestFragment_Int>();
+				// forces handle reservation
+				FMassEntityHandle EntityA = BuilderA.Commit();
+
+				FMassEntityHandle EntityB;
+				{
+					FEntityBuilder BuilderB(*EntityManager.Get());
+					BuilderB.Add<FTestFragment_Float>();
+					EntityB = BuilderB.Commit();
+					BuilderA = MoveTemp(BuilderB);
+				}
+
+				// at this point EntityB should be valid while the original EntityA not
+				AITEST_TRUE("The original entity is valid", EntityManager->IsEntityValid(EntityA));
+				AITEST_TRUE("The original entity is active", EntityManager->IsEntityActive(EntityA));
+				AITEST_TRUE("The secondary entity entity is valid", EntityManager->IsEntityValid(EntityB));
+				AITEST_TRUE("The secondary entity is active", EntityManager->IsEntityActive(EntityB));
+
+				{
+					AITEST_SCOPED_CHECK("Trying to commit an already committed", 1);
+					BuilderA.Commit();
+				}
+				FMassEntityHandle EntityA2 = BuilderA.GetEntityHandle();
+				AITEST_TRUE("The entity handle is the same as the builder's that has been moved", EntityA2 == EntityB);
+			}
+
+			return true;
+		}
+	};
+	IMPLEMENT_AI_INSTANT_TEST(FPassOver, "System.Mass.EntityBuilder.PassOver");
+} // UE::Mass::Test::EntityBuilder
+
+UE_ENABLE_OPTIMIZATION_SHIP
+
+#undef LOCTEXT_NAMESPACE
diff --git a/Engine/Source/Runtime/MassEntity/Private/MassEntityBuilder.cpp b/Engine/Source/Runtime/MassEntity/Private/MassEntityBuilder.cpp
new file mode 100644
--- /dev/null
+++ b/Engine/Source/Runtime/MassEntity/Private/MassEntityBuilder.cpp
@@ -0,0 +1,192 @@
+// Copyright Epic Games, Inc. All Rights Reserved.
+
+#include "MassEntityBuilder.h"
+#include "MassEntityManager.h"
+#include "VisualLogger/VisualLogger.h"
+
+
+namespace UE::Mass
+{
+//-----------------------------------------------------------------------------
+// FEntityBuilder
+//-----------------------------------------------------------------------------
+FEntityBuilder::FEntityBuilder(FMassEntityManager& InEntityManager)
+	: EntityManager(InEntityManager.AsShared())
+{	
+}
+
+FEntityBuilder::FEntityBuilder(const TSharedRef<FMassEntityManager>& InEntityManager)
+	: EntityManager(InEntityManager)
+{
+	
+}
+
+FEntityBuilder::FEntityBuilder(const FEntityBuilder& Other)
+	: EntityManager(Other.EntityManager)
+{
+	*this = Other;
+}
+
+FEntityBuilder& FEntityBuilder::operator=(const FEntityBuilder& Other)
+{
+	if (testableEnsureMsgf(Other.IsValid(), TEXT("Copying invalid entity builder instances is not supported")))
+	{
+		// if we already have an EntityHandle reserved we might want to keep it - why reserve a handle again
+		// soon, the reserved handle doesn't have an archetype associated with it?
+		// We do need to release the handle if we're dealing with a different entity manager (unexpected in practice, but possible [for now])
+		if (EntityManager != Other.EntityManager)
+		{
+			ConditionallyReleaseEntityHandle();
+			EntityManager = Other.EntityManager;
+		}
+		// we also reset the handle if this builder has already committed its entity - the entity needs to 
+		// be forgotten by this builder, it's "out in the wild" now and should be safe from accidental destruction.
+		else if (State == EState::Committed)
+		{
+			EntityHandle.Reset();
+		}
+			
+		Composition = Other.Composition;
+		Fragments = Other.Fragments;
+		SharedFragments = Other.SharedFragments;
+		ConstSharedFragments = Other.ConstSharedFragments;
+
+		State = Composition.IsEmpty() ? EState::Empty : EState::ReadyToCommit;
+	}
+
+	return *this;
+}
+
+FEntityBuilder& FEntityBuilder::operator=(FEntityBuilder&& Other)
+{
+	if (testableEnsureMsgf(Other.IsValid(), TEXT("Copying invalid entity builder instances is not supported")))
+	{
+		// if we already have an EntityHandle reserved we might want to keep it - why reserve a handle again
+		// soon, the reserved handle doesn't have an archetype associated with it?
+		// We do need to release the handle if we're dealing with a different entity manager (unexpected in practice, but possible [for now])
+		if (EntityManager != Other.EntityManager)
+		{
+			ConditionallyReleaseEntityHandle();
+			EntityManager = MoveTemp(Other.EntityManager);
+		}
+		Fragments = MoveTemp(Other.Fragments);
+		SharedFragments = MoveTemp(Other.SharedFragments);
+		ConstSharedFragments = MoveTemp(Other.ConstSharedFragments);
+
+		// the main point of the elaborated logic below is to avoid needlessly releasing reserved entities.
+		if (HasReservedEntityHandle())
+		{
+			if (Other.HasReservedEntityHandle())
+			{
+				ConditionallyReleaseEntityHandle();
+				EntityHandle = Other.EntityHandle;
+			}
+			State = (Other.State == EState::Committed)
+				? EState::ReadyToCommit // we have a reserved entity at hand, we can Commit again
+				: Other.State;
+		}
+		else
+		{
+			// we just take everything as is
+			EntityHandle = Other.EntityHandle;
+			State = Other.State;
+		}
+
+		Other.EntityHandle.Reset();
+		Other.State = EState::Invalid;
+	}
+
+	return *this;
+}
+
+FEntityBuilder::~FEntityBuilder()
+{
+	ConditionallyReleaseEntityHandle();
+}
+
+FMassEntityHandle FEntityBuilder::Commit()
+{
+	if (!testableEnsureMsgf(State != EState::Committed, TEXT("Trying to commit an already committed EntityBuilder. This request will be ignored.")))
+	{
+		return EntityHandle;
+	}
+	if (Composition.IsEmpty())
+	{
+		UE_VLOG_ALWAYS_UELOG(EntityManager->GetOwner(), LogMass, Warning, TEXT("%hs: Attempting to commit while no composition has been configured."), __FUNCTION__);
+		UE_CVLOG_ALWAYS_UELOG(EntityHandle.IsValid(), EntityManager->GetOwner(), LogMass, Error, TEXT("Failing to commit while the entity handle has already been reserved."));
+		return FMassEntityHandle();
+	}
+
+	CacheEntityHandle();
+	
+	if (!ensureMsgf(EntityManager->IsProcessing() == false, TEXT("Committing EntityBuilder while Mass processing is in progress is not supported yet.")))
+	{
+		// @todo we need to issue commands in this case
+	}
+	else
+	{
+		// directly create the entity
+		TSharedRef<FMassEntityManager::FEntityCreationContext> CreationContext = EntityManager->GetOrMakeCreationContext();
+
+		FMassArchetypeSharedFragmentValues SharedFragmentValues;
+		for (FInstancedStruct& SharedFragmentInstance : SharedFragments)
+		{
+			check(SharedFragmentInstance.IsValid());
+			const FSharedStruct& SharedStruct = EntityManager->GetOrCreateSharedFragment(*SharedFragmentInstance.GetScriptStruct(), SharedFragmentInstance.GetMemory());
+			SharedFragmentValues.AddSharedFragment(SharedStruct);
+		}
+		for (FInstancedStruct& ConstSharedFragmentInstance : ConstSharedFragments)
+		{
+			check(ConstSharedFragmentInstance.IsValid());
+			const FSharedStruct& ConstSharedStruct = EntityManager->GetOrCreateSharedFragment(*ConstSharedFragmentInstance.GetScriptStruct(), ConstSharedFragmentInstance.GetMemory());
+			SharedFragmentValues.AddConstSharedFragment(ConstSharedStruct);
+		}
+
+		FMassArchetypeHandle ArchetypeHandle = EntityManager->CreateArchetype(Composition);
+		EntityManager->BuildEntity(EntityHandle, ArchetypeHandle, SharedFragmentValues);
+		EntityManager->SetEntityFragmentsValues(EntityHandle, Fragments);
+	}
+
+	State = EState::Committed;
+
+	return EntityHandle;
+}
+
+void FEntityBuilder::Reset()
+{
+	ConditionallyReleaseEntityHandle();
+
+	State = EState::Empty;
+
+	Composition.Reset();
+	Fragments.Reset();
+	SharedFragments.Reset();
+	ConstSharedFragments.Reset();
+}
+
+FMassEntityHandle FEntityBuilder::GetEntityHandle() const
+{
+	CacheEntityHandle();
+	return EntityHandle;
+}
+
+void FEntityBuilder::ConditionallyReleaseEntityHandle()
+{
+	if ((EntityHandle.IsValid() == true) && (State != EState::Committed))
+	{
+		EntityManager->ReleaseReservedEntity(EntityHandle);
+	}
+
+	EntityHandle.Reset();
+}
+
+void FEntityBuilder::CacheEntityHandle() const
+{
+	if (EntityHandle.IsValid() == false)
+	{
+		checkf(State != EState::Committed, TEXT("Reserving an entity while the builder has already committed. This should not happen. Indicates an error during builder copying from another instance."))
+		EntityHandle = EntityManager->ReserveEntity();
+	}
+}
+
+} // namespace UE::Mass
\ No newline at end of file
diff --git a/Engine/Source/Runtime/MassEntity/Private/MassEntityManager.cpp b/Engine/Source/Runtime/MassEntity/Private/MassEntityManager.cpp
--- a/Engine/Source/Runtime/MassEntity/Private/MassEntityManager.cpp
+++ b/Engine/Source/Runtime/MassEntity/Private/MassEntityManager.cpp
@@ -14,6 +14,8 @@
 #include "Misc/CoreDelegates.h"
 #include "Algo/Find.h"
 #include "MassEntityUtils.h"
+#include "MassEntityBuilder.h"
+
 
 const FMassEntityHandle FMassEntityManager::InvalidEntity;
 
@@ -2051,6 +2053,11 @@ bool FMassEntityManager::DirtyCreationContext()
 	return false;
 }
 
+UE::Mass::FEntityBuilder FMassEntityManager::MakeEntityBuilder()
+{
+	return UE::Mass::FEntityBuilder(AsShared());
+}
+
 bool FMassEntityManager::DebugDoCollectionsOverlapCreationContext(TConstArrayView<FMassArchetypeEntityCollection> EntityCollections) const
 {
 	if (TSharedPtr<FEntityCreationContext> AsSharedPtr = ActiveCreationContext.Pin())
diff --git a/Engine/Source/Runtime/MassEntity/Public/MassEntityBuilder.h b/Engine/Source/Runtime/MassEntity/Public/MassEntityBuilder.h
new file mode 100644
--- /dev/null
+++ b/Engine/Source/Runtime/MassEntity/Public/MassEntityBuilder.h
@@ -0,0 +1,305 @@
+// Copyright Epic Games, Inc. All Rights Reserved.
+
+#pragma once
+
+#include "MassEntityConcepts.h"
+#include "MassEntityHandle.h"
+#include "MassEntityTypes.h"
+
+
+struct FMassEntityManager;
+
+namespace UE::Mass
+{
+
+/**
+ * FEntityBuilder is a utility struct that provides a convenient way to create and configure entities
+ * in the Mass framework. It bridges multiple APIs from FMassEntityManager, MassSpawnerSubsystem,
+ * MassEntityTemplates, and other related components, allowing for streamlined entity creation and configuration.
+ *
+ * Key Features:
+ * - Can be seamlessly used in place of FMassEntityHandle, allowing for consistent and intuitive usage.
+ * - An entity only gets created once Commit() is called
+ * - Copyable, but copied instances represent new entities without carrying over the reserved entity handle.
+ *
+ * Example Usage:
+ * {
+ * 		FEntityBuilder Builder(EntityManager);
+ * 		Builder.Add<FTransformFragment>(FTransform(FVector(100, 200, 300)))
+ * 			.Commit();	// the entity gets reserved and built by this call
+ * } 
+ *
+ * {
+ *     FEntityBuilder Builder(EntityManager);
+ *     FMassEntityHandle ReservedEntity = Builder; // Entity handle reserved, can be used for commands.
+ *     Builder.Add_GetRef<FTransformFragment>().GetMutableTransform().SetTranslation(FVector(100, 200, 300));
+ *     Builder.Commit(); // Entity creation is finalized at this point.
+ * }
+ *
+ * // Example of chaining with FMassEntityManager's MakeEntityBuilder() method:
+ * FMassEntityHandle NewEntity = EntityManager.MakeEntityBuilder()
+ *     .Add<FMassStaticRepresentationTag>()
+ *     .Add<FTransformFragment>()
+ *     .Add<FAgentRadiusFragment>(FAgentRadiusFragment{ .Radius = 35.f})
+ *     .Add<FMassVelocityFragment>()
+ *     .Commit();
+ *
+ * Current Limitations:
+ * - Committing entities while Mass's processing is in progress is not yet supported; this functionality will be implemented in the near future.
+ * - no support for entity grouping
+ */
+struct MASSENTITY_API FEntityBuilder
+{
+	/** Constructs a FEntityBuilder using a reference to a FMassEntityManager. */
+	explicit FEntityBuilder(FMassEntityManager& InEntityManager);
+
+	/** Constructs a FEntityBuilder using a shared reference to a FMassEntityManager. */
+	explicit FEntityBuilder(const TSharedRef<FMassEntityManager>& InEntityManager);
+
+	/** Copy constructor - copies-create a new instance that represents a new entity and does not carry over reserved handle. */
+	FEntityBuilder(const FEntityBuilder& Other);
+
+	/** Assignment operator - copies represent new entities, with no carryover of reserved handle from the original. */
+	FEntityBuilder& operator=(const FEntityBuilder& Other);
+
+	/** Move assignment operator - moves over all the data from Other, including the internal state (like whether the entity handle has already been reserved) */
+	FEntityBuilder& operator=(FEntityBuilder&& Other);
+
+	/** Destructor - automatically commits entity creation if not explicitly aborted or committed beforehand. */
+	~FEntityBuilder();
+
+	/**
+	 * Finalizes the creation of the entity with the specified fragments and configurations.
+	 * Note that this function needs to be called manually, no automated entity creation will take place upon builder's destruction. 
+	 */
+	FMassEntityHandle Commit();
+
+	/** Resets the builder to its initial state, discarding all previous entity configurations. */
+	void Reset();
+
+	/**
+	 * Adds a fragment of type T to the entity and returns a reference to it, constructing it with the provided arguments.
+	 * The function will assert if an element of type T already exists.
+	 * @param T - The type of fragment to add.
+	 * @param InArgs - Constructor arguments for initializing the fragment.
+	 * @return A reference to the added fragment.
+	 */
+	template<typename T, typename... TArgs>
+	T& Add_GetRef(TArgs&&... InArgs) requires (!(CTag<T> || CChunkFragment<T>));
+
+	/**
+	 * Adds a fragment of type T to the entity and returns a reference to it, constructing it with the provided arguments.
+	 * If a fragment of the given type already exists then it will be overriden and its reference returned.
+	 * @return A reference to the added fragment.
+	 */
+	template<typename T, typename... TArgs>
+	T& GetOrCreate(TArgs&&... InArgs) requires (!(CTag<T> || CChunkFragment<T>));
+
+	/**
+	 * Adds a tag of type T to the entity.
+	 * @return Reference to this FEntityBuilder for method chaining.
+	 */
+	template<CTag T>
+	FEntityBuilder& Add();
+
+	/**
+	 * Adds a chunk fragment of type T to the entity.
+	 * @return Reference to this FEntityBuilder for method chaining.
+	 */
+	template<CChunkFragment T>
+	FEntityBuilder& Add();
+
+	/**
+	 * Adds a fragment of type T to the entity, constructing it with the provided arguments.
+	 * @return Reference to this FEntityBuilder for method chaining.
+	 */
+	template<typename T, typename... TArgs>
+	FEntityBuilder& Add(TArgs&&... InArgs) requires (!(CTag<T> || CChunkFragment<T>));
+
+	/**
+	 * Finds and retrieves a pointer to a fragment of type T if it exists.
+	 * @param T - The type of fragment to find.
+	 * @return Pointer to the fragment, or nullptr if it does not exist.
+	 */
+	template<typename T>
+	T* Find();
+
+	/** Converts the builder to a FMassEntityHandle, reserving the entity handle if not already committed. */
+	[[nodiscard]] FMassEntityHandle GetEntityHandle() const;
+
+	/** Checks whether the builder is in a valid, expected state */
+	bool IsValid() const;
+
+	/** @return whether the builder has an entity handle reserved and the data has not been committed yet */
+	bool HasReservedEntityHandle() const;
+
+	/** @return whether the builder has already committed the data */
+	bool IsCommitted() const;
+
+private:
+	TSharedRef<FMassEntityManager> EntityManager;
+	mutable FMassEntityHandle EntityHandle;
+
+	FMassArchetypeCompositionDescriptor Composition;
+
+	template<CFragment T> 
+	TArray<FInstancedStruct>& GetInstancedStructContainer() 
+	{ 
+		return Fragments; 
+	}
+
+	template<CSharedFragment T> 
+	TArray<FInstancedStruct>& GetInstancedStructContainer() 
+	{ 
+		return SharedFragments; 
+	}
+
+	template<CConstSharedFragment T> 
+	TArray<FInstancedStruct>& GetInstancedStructContainer() 
+	{ 
+		return ConstSharedFragments; 
+	}
+
+	template<typename T>
+	auto& GetBitSetBuilder()
+	{
+		return Composition.GetContainer<T>();
+	}
+
+	/** Releases reserved handle if it has not been committed yet */
+	void ConditionallyReleaseEntityHandle();
+
+	void CacheEntityHandle() const;
+
+	TArray<FInstancedStruct> Fragments;
+	TArray<FInstancedStruct> SharedFragments;
+	TArray<FInstancedStruct> ConstSharedFragments;
+
+	enum class EState : uint8
+	{
+		Empty,
+		ReadyToCommit,
+		Committed,
+		Invalid,
+	};
+	EState State = EState::Empty;
+};
+
+struct FScopedEntityBuilder : FEntityBuilder
+{
+	UE_NONCOPYABLE(FScopedEntityBuilder);
+
+	template<typename... TArgs>
+	FScopedEntityBuilder(TArgs&&... InArgs)
+		: FEntityBuilder(Forward<TArgs>(InArgs)...)
+	{	
+	}
+
+	~FScopedEntityBuilder()
+	{
+		Commit();
+	}
+};
+
+//-----------------------------------------------------------------------------
+// Inlines and specializations
+//-----------------------------------------------------------------------------
+template<CTag T>
+FEntityBuilder& FEntityBuilder::Add()
+{
+	Composition.Tags.Add<T>();
+	State = EState::ReadyToCommit;
+	return *this;
+}
+
+template<CChunkFragment T>
+FEntityBuilder& FEntityBuilder::Add()
+{
+	Composition.ChunkFragments.Add<T>();
+	State = EState::ReadyToCommit;
+	return *this;
+}
+
+template<typename T, typename... TArgs>
+FEntityBuilder& FEntityBuilder::Add(TArgs&&... InArgs) requires (!(CTag<T> || CChunkFragment<T>))
+{
+	Add_GetRef<T>(InArgs...);
+	return *this;
+}
+
+template<typename T, typename... TArgs>
+T& FEntityBuilder::Add_GetRef(TArgs&&... InArgs) requires (!(CTag<T> || CChunkFragment<T>))
+{
+	State = EState::ReadyToCommit;
+
+	if (ensureMsgf(GetBitSetBuilder<T>().template Contains<T>() == false, TEXT("Element of type %s has already been added"), *T::StaticStruct()->GetName()))
+	{
+		GetBitSetBuilder<T>().template Add<T>();
+		return GetInstancedStructContainer<T>().Add_GetRef(FInstancedStruct::Make<T>(InArgs...)).template GetMutable<T>();
+	}
+
+	return GetInstancedStructContainer<T>().FindByPredicate([](const FInstancedStruct& ExistingElement)
+		{
+			return ExistingElement.GetScriptStruct() == T::StaticStruct();
+		})->template GetMutable<T>();
+}
+
+template<typename T, typename... TArgs>
+T& FEntityBuilder::GetOrCreate(TArgs&&... InArgs) requires (!(CTag<T> || CChunkFragment<T>))
+{
+	State = EState::ReadyToCommit;
+
+	FInstancedStruct* ElementFound = nullptr;
+	if (GetBitSetBuilder<T>().template Contains<T>())
+	{
+		// replace
+		ElementFound = GetInstancedStructContainer<T>().FindByPredicate([](const FInstancedStruct& ExistingElement)
+		{
+			return ExistingElement.GetScriptStruct() == T::StaticStruct();
+		});
+
+		checkf(ElementFound, TEXT("We expect the element to be found since we already tested the Composition"));
+		ElementFound->GetMutable<T>() = T(InArgs...);
+	}
+	else
+	{
+		GetBitSetBuilder<T>().template Add<T>();
+		ElementFound = &GetInstancedStructContainer<T>().Add_GetRef(FInstancedStruct::Make<T>(InArgs...));
+	}
+
+	return ElementFound->GetMutable<T>();
+}
+
+template<typename T>
+T* FEntityBuilder::Find()
+{
+	if (GetBitSetBuilder<T>().template Contains<T>())
+	{
+		FInstancedStruct* ElementFound = GetInstancedStructContainer<T>().FindByPredicate([](const FInstancedStruct& Element)
+		{
+			return Element.GetScriptStruct() == T::StaticStruct();
+		});
+
+		checkf(ElementFound, TEXT("We expect the element to be found since we already tested the Composition"));
+		return ElementFound->GetMutablePtr<T>();
+	}
+
+	return nullptr;
+}
+
+inline bool FEntityBuilder::IsValid() const
+{
+	return State != EState::Invalid;
+}
+
+inline bool FEntityBuilder::HasReservedEntityHandle() const
+{
+	return State != EState::Committed && EntityHandle.IsValid();
+}
+
+inline bool FEntityBuilder::IsCommitted() const
+{
+	return State == EState::Committed;
+}
+
+} // namespace UE::Mass
diff --git a/Engine/Source/Runtime/MassEntity/Public/MassEntityConcepts.h b/Engine/Source/Runtime/MassEntity/Public/MassEntityConcepts.h
new file mode 100644
--- /dev/null
+++ b/Engine/Source/Runtime/MassEntity/Public/MassEntityConcepts.h
@@ -0,0 +1,27 @@
+// Copyright Epic Games, Inc. All Rights Reserved.
+
+#pragma once
+
+#include "MassEntityElementTypes.h"
+
+
+namespace UE::Mass
+{
+	template <typename T>
+	concept CFragment = TIsDerivedFrom<typename TRemoveReference<T>::Type, FMassFragment>::Value;
+
+	template <typename T>
+	concept CTag = TIsDerivedFrom<typename TRemoveReference<T>::Type, FMassTag>::Value;
+
+	template <typename T>
+	concept CChunkFragment = TIsDerivedFrom<typename TRemoveReference<T>::Type, FMassChunkFragment>::Value;
+
+	template <typename T>
+	concept CSharedFragment = TIsDerivedFrom<typename TRemoveReference<T>::Type, FMassSharedFragment>::Value;
+
+	template <typename T>
+	concept CConstSharedFragment = TIsDerivedFrom<typename TRemoveReference<T>::Type, FMassConstSharedFragment>::Value;
+
+	template <typename T>
+	concept CElement = CFragment<T> || CTag<T> || CChunkFragment<T> || CSharedFragment<T> || CConstSharedFragment<T>;
+}
diff --git a/Engine/Source/Runtime/MassEntity/Public/MassEntityManager.h b/Engine/Source/Runtime/MassEntity/Public/MassEntityManager.h
--- a/Engine/Source/Runtime/MassEntity/Public/MassEntityManager.h
+++ b/Engine/Source/Runtime/MassEntity/Public/MassEntityManager.h
@@ -29,9 +29,13 @@ class FOutputDevice;
 struct FMassDebugger;
 enum class EMassFragmentAccess : uint8;
 enum class EForkProcessRole : uint8;
-namespace UE::Mass::Private
+namespace UE::Mass
 {
-	struct FEntityStorageInitializer;
+	struct FEntityBuilder;
+	namespace Private
+	{
+		struct FEntityStorageInitializer;
+	}
 }
 
 // use REQUESTED_MASS_CONCURRENT_RESERVE=0 in your project's Build.cs file to disable concurrent storage
@@ -1015,6 +1019,8 @@ PRAGMA_ENABLE_DEPRECATION_WARNINGS
 		}
 	}
 
+	[[nodiscard]] UE::Mass::FEntityBuilder MakeEntityBuilder();
+
 	FMassObserverManager& GetObserverManager() { return ObserverManager; }
 
 	FOnNewArchetypeDelegate& GetOnNewArchetypeEvent() { return OnNewArchetypeEvent; }
diff --git a/Engine/Source/Runtime/MassEntity/Public/MassEntityTypes.h b/Engine/Source/Runtime/MassEntity/Public/MassEntityTypes.h
--- a/Engine/Source/Runtime/MassEntity/Public/MassEntityTypes.h
+++ b/Engine/Source/Runtime/MassEntity/Public/MassEntityTypes.h
@@ -11,6 +11,7 @@
 #endif // UE_ENABLE_INCLUDE_ORDER_DEPRECATED_IN_5_6
 #include "StructUtils/SharedStruct.h"
 #include "MassEntityElementTypes.h"
+#include "MassEntityConcepts.h"
 #ifdef WITH_AITESTSUITE
 #include "TestableEnsures.h"
 #else
@@ -139,6 +140,27 @@ struct FMassArchetypeCompositionDescriptor
 
 	MASSENTITY_API void DebugOutputDescription(FOutputDevice& Ar) const;
 
+	template<UE::Mass::CFragment T>
+	FMassFragmentBitSet& GetContainer();
+
+	template<UE::Mass::CTag T>
+	FMassTagBitSet& GetContainer();
+
+	template<UE::Mass::CChunkFragment T>
+	FMassChunkFragmentBitSet& GetContainer();
+
+	template<UE::Mass::CSharedFragment T>
+	FMassSharedFragmentBitSet& GetContainer();
+
+	template<UE::Mass::CConstSharedFragment T>
+	FMassConstSharedFragmentBitSet& GetContainer();
+
+	template<typename T>
+	bool Contains() const;
+
+	template<typename T>
+	bool Add();
+
 	FMassFragmentBitSet Fragments;
 	FMassTagBitSet Tags;
 	FMassChunkFragmentBitSet ChunkFragments;
@@ -632,3 +654,47 @@ struct FMassArchetypeCreationParams
 	FName DebugName;
 };
 
+//-----------------------------------------------------------------------------
+// INLINES
+//-----------------------------------------------------------------------------
+template<UE::Mass::CFragment T>
+FMassFragmentBitSet& FMassArchetypeCompositionDescriptor::GetContainer()
+{
+	return Fragments;
+}
+
+template<UE::Mass::CTag T>
+FMassTagBitSet& FMassArchetypeCompositionDescriptor::GetContainer()
+{
+	return Tags;
+}
+
+template<UE::Mass::CChunkFragment T>
+FMassChunkFragmentBitSet& FMassArchetypeCompositionDescriptor::GetContainer()
+{
+	return ChunkFragments;
+}
+
+template<UE::Mass::CSharedFragment T>
+FMassSharedFragmentBitSet& FMassArchetypeCompositionDescriptor::GetContainer()
+{
+	return SharedFragments;
+}
+
+template<UE::Mass::CConstSharedFragment T>
+FMassConstSharedFragmentBitSet& FMassArchetypeCompositionDescriptor::GetContainer()
+{
+	return ConstSharedFragments;
+}
+
+template<typename T>
+bool FMassArchetypeCompositionDescriptor::Contains() const
+{
+	return GetContainer<T>().template Contains<T>();
+}
+
+template<typename T>
+bool FMassArchetypeCompositionDescriptor::Add()
+{
+	return GetContainer<T>().template Add<T>();
+}
