[Mass] Issue 687211: UE::Mass::Test::EntityBuilder compilation fix

diff --git a/Engine/Source/Developer/MassEntityTestSuite/Private/MassEntityBuilderTest.cpp b/Engine/Source/Developer/MassEntityTestSuite/Private/MassEntityBuilderTest.cpp
--- a/Engine/Source/Developer/MassEntityTestSuite/Private/MassEntityBuilderTest.cpp
+++ b/Engine/Source/Developer/MassEntityTestSuite/Private/MassEntityBuilderTest.cpp
@@ -22,7 +22,9 @@ namespace UE::Mass::Test::EntityBuilder
 	{
 		virtual bool InstantTest() override
 		{
+#if  WITH_MASSENTITY_DEBUG
 			int32 EntitiesCreated = EntityManager->DebugGetArchetypeEntitiesCount(IntsArchetype);
+#endif // WITH_MASSENTITY_DEBUG
 			int32 SomeCounter = 1;
 			int32 EntitiesCreatedThisStep = 0;
 
@@ -32,10 +34,10 @@ namespace UE::Mass::Test::EntityBuilder
 				EntityBuilder.Commit();
 				EntitiesCreatedThisStep = 1;
 			}
-
+#if WITH_MASSENTITY_DEBUG
 			AITEST_EQUAL("Number of entities created with basic use", EntityManager->DebugGetArchetypeEntitiesCount(IntsArchetype) - EntitiesCreated, EntitiesCreatedThisStep);
 			EntitiesCreated = EntityManager->DebugGetArchetypeEntitiesCount(IntsArchetype);
-
+#endif // WITH_MASSENTITY_DEBUG
 			{
 				FEntityBuilder EntityBuilder(*EntityManager.Get());
 				ON_SCOPE_EXIT
@@ -46,19 +48,19 @@ namespace UE::Mass::Test::EntityBuilder
 				EntityBuilder.Add_GetRef<FTestFragment_Int>(SomeCounter++);
 			}
 			EntitiesCreatedThisStep = 1;
-
+#if WITH_MASSENTITY_DEBUG
 			AITEST_EQUAL("Number of entities created with ON_SCOPE_EXIT", EntityManager->DebugGetArchetypeEntitiesCount(IntsArchetype) - EntitiesCreated, EntitiesCreatedThisStep);
 			EntitiesCreated = EntityManager->DebugGetArchetypeEntitiesCount(IntsArchetype);
-
+#endif // WITH_MASSENTITY_DEBUG
 			{
 				FScopedEntityBuilder EntityBuilder(*EntityManager.Get());
 				EntityBuilder.Add_GetRef<FTestFragment_Int>(SomeCounter++);
 			}
 			EntitiesCreatedThisStep = 1;
-
+#if WITH_MASSENTITY_DEBUG
 			AITEST_EQUAL("Number of entities created with scoped builder", EntityManager->DebugGetArchetypeEntitiesCount(IntsArchetype) - EntitiesCreated, EntitiesCreatedThisStep);
 			EntitiesCreated = EntityManager->DebugGetArchetypeEntitiesCount(IntsArchetype);
-
+#endif // WITH_MASSENTITY_DEBUG
 			{
 				FScopedEntityBuilder EntityBuilder(*EntityManager.Get());
 				EntityBuilder.Add_GetRef<FTestFragment_Int>(SomeCounter++);
@@ -67,10 +69,10 @@ namespace UE::Mass::Test::EntityBuilder
 				EntityBuilder2.Commit();
 			}
 			EntitiesCreatedThisStep = 2;
-
+#if WITH_MASSENTITY_DEBUG
 			AITEST_EQUAL("Number of entities created with with a scoped builder and its regular copy", EntityManager->DebugGetArchetypeEntitiesCount(IntsArchetype) - EntitiesCreated, EntitiesCreatedThisStep);
 			EntitiesCreated = EntityManager->DebugGetArchetypeEntitiesCount(IntsArchetype);
-
+#endif // WITH_MASSENTITY_DEBUG
 			{
 				FEntityBuilder EntityBuilder(*EntityManager.Get());
 				EntityBuilder.Add_GetRef<FTestFragment_Int>(SomeCounter++);
@@ -79,10 +81,10 @@ namespace UE::Mass::Test::EntityBuilder
 				FEntityBuilder EntityBuilder2 = EntityBuilder;
 			}
 			EntitiesCreatedThisStep = 1;
-
+#if WITH_MASSENTITY_DEBUG
 			AITEST_EQUAL("Number of entities created with Commit and an abandoned copy of a builder", EntityManager->DebugGetArchetypeEntitiesCount(IntsArchetype) - EntitiesCreated, EntitiesCreatedThisStep);
 			EntitiesCreated = EntityManager->DebugGetArchetypeEntitiesCount(IntsArchetype);
-
+#endif // WITH_MASSENTITY_DEBUG
 			{
 				FEntityBuilder EntityBuilder(*EntityManager.Get());
 				EntityBuilder.Add_GetRef<FTestFragment_Int>(SomeCounter++);
@@ -93,10 +95,10 @@ namespace UE::Mass::Test::EntityBuilder
 				EntityBuilder.Reset();
 			}
 			EntitiesCreatedThisStep = 1;
-
+#if WITH_MASSENTITY_DEBUG
 			AITEST_EQUAL("Number of entities created with committed copy and abandoned original builder", EntityManager->DebugGetArchetypeEntitiesCount(IntsArchetype) - EntitiesCreated, EntitiesCreatedThisStep);
 			EntitiesCreated = EntityManager->DebugGetArchetypeEntitiesCount(IntsArchetype);
-
+#endif // WITH_MASSENTITY_DEBUG
 			return true;
 		}
 	};
