[Mass] spelling fixes

#rnx

diff --git a/Engine/Source/Developer/MassEntityTestSuite/Private/MassParallelQueryTest.cpp b/Engine/Source/Developer/MassEntityTestSuite/Private/MassParallelQueryTest.cpp
--- a/Engine/Source/Developer/MassEntityTestSuite/Private/MassParallelQueryTest.cpp
+++ b/Engine/Source/Developer/MassEntityTestSuite/Private/MassParallelQueryTest.cpp
@@ -31,24 +31,24 @@ struct FQueryTest_ParallelForBasic : FEntityTestBase
 		Query.AddRequirement<FTestFragment_Float>(EMassFragmentAccess::ReadWrite);
 
 		std::atomic<int32> EntitiesProcessed = 0;
-		std::atomic<int32> ConcurrentAcces = 0;
+		std::atomic<int32> ConcurrentAccess = 0;
 		FCriticalSection ParallelDetectionCS;
 
 		FMassExecutionContext Context = EntityManager->CreateExecutionContext(/*DeltaSeconds=*/0.f);
-		Query.ParallelForEachEntityChunk(*EntityManager, Context, [&ParallelDetectionCS, &EntitiesProcessed, &ConcurrentAcces](FMassExecutionContext& Context)
+		Query.ParallelForEachEntityChunk(*EntityManager, Context, [&ParallelDetectionCS, &EntitiesProcessed, &ConcurrentAccess](FMassExecutionContext& Context)
 			{
 				EntitiesProcessed += Context.GetNumEntities();
 
 				if (ParallelDetectionCS.TryLock() == false)
 				{
-					++ConcurrentAcces;
+					++ConcurrentAccess;
 					return;
 				}
 				FPlatformProcess::Sleep(0.01f);
 			}, FMassEntityQuery::ForceParallelExecution);
 		
 		AITEST_EQUAL("All created entites should have been processed at this point", EntitiesProcessed.load(), NumEntitiesToCreate);
-		AITEST_TRUE("Some work is expected to be executed in parallel", ConcurrentAcces.load() > 0);
+		AITEST_TRUE("Some work is expected to be executed in parallel", ConcurrentAccess.load() > 0);
 
 		return true;
 	}
@@ -84,21 +84,21 @@ struct FQueryTest_ParallelForCollection : FEntityTestBase
 		Query.AddRequirement<FTestFragment_Float>(EMassFragmentAccess::ReadWrite);
 
 		std::atomic<int32> EntitiesProcessed = 0;
-		std::atomic<int32> ConcurrentAcces = 0;
+		std::atomic<int32> ConcurrentAccess = 0;
 		std::atomic<int32> JobsExecuted = 0;
 		FCriticalSection ParallelDetectionCS;
 
 		FMassExecutionContext Context = EntityManager->CreateExecutionContext(/*DeltaSeconds=*/0.f);
 		Context.SetEntityCollection(EntityCollection);
 
-		Query.ParallelForEachEntityChunk(*EntityManager, Context, [&JobsExecuted, &ParallelDetectionCS, &EntitiesProcessed, &ConcurrentAcces](FMassExecutionContext& Context)
+		Query.ParallelForEachEntityChunk(*EntityManager, Context, [&JobsExecuted, &ParallelDetectionCS, &EntitiesProcessed, &ConcurrentAccess](FMassExecutionContext& Context)
 			{
 				++JobsExecuted;
 				EntitiesProcessed += Context.GetNumEntities();
 
 				if (ParallelDetectionCS.TryLock() == false)
 				{
-					++ConcurrentAcces;
+					++ConcurrentAccess;
 					return;
 				}
 				FPlatformProcess::Sleep(0.01f);
@@ -106,7 +106,7 @@ struct FQueryTest_ParallelForCollection : FEntityTestBase
 
 		AITEST_EQUAL("All entities passed for processing should have been counted", EntitiesProcessed.load(), EntitiesToProcess.Num());
 		AITEST_EQUAL("The number of jobs should match expectations", NumJobs, JobsExecuted.load());
-		AITEST_TRUE("Some work is expected to be executed in parallel", ConcurrentAcces.load() > 0);
+		AITEST_TRUE("Some work is expected to be executed in parallel", ConcurrentAccess.load() > 0);
 
 		return true;
 	}
diff --git a/Engine/Source/Runtime/MassEntity/Public/MassEntityHandle.h b/Engine/Source/Runtime/MassEntity/Public/MassEntityHandle.h
--- a/Engine/Source/Runtime/MassEntity/Public/MassEntityHandle.h
+++ b/Engine/Source/Runtime/MassEntity/Public/MassEntityHandle.h
@@ -76,5 +76,5 @@ struct alignas(8) FMassEntityHandle
 	}
 };
 
-static_assert(sizeof(FMassEntityHandle) == sizeof(uint64), "Expected FMassEntityHandle to be convertable to a 64-bit integer value, so size needs to be 8 bytes.");
-static_assert(alignof(FMassEntityHandle) == sizeof(uint64), "Expected FMassEntityHandle to be convertable to a 64-bit integer value, so alignment needs to be 8 bytes.");
+static_assert(sizeof(FMassEntityHandle) == sizeof(uint64), "Expected FMassEntityHandle to be convertible to a 64-bit integer value, so size needs to be 8 bytes.");
+static_assert(alignof(FMassEntityHandle) == sizeof(uint64), "Expected FMassEntityHandle to be convertible to a 64-bit integer value, so alignment needs to be 8 bytes.");
