[Mass] fixed unit tests failing in MassDependencyTest.cpp due to not being updated after entity query API changes (deprecating the MassEntityManager as function parameter)

#rnx

diff --git a/Engine/Source/Developer/MassEntityTestSuite/Private/MassDependencyTest.cpp b/Engine/Source/Developer/MassEntityTestSuite/Private/MassDependencyTest.cpp
--- a/Engine/Source/Developer/MassEntityTestSuite/Private/MassDependencyTest.cpp
+++ b/Engine/Source/Developer/MassEntityTestSuite/Private/MassDependencyTest.cpp
@@ -200,12 +200,12 @@ struct FComplexScenario : FDependencySolverBase
 		}
 
 		{
-			UMassTestProcessorBase* Proc = Processors.Add_GetRef(NewObject<UMassTestProcessor_E>());
+			UMassTestProcessorBase* Proc = Processors.Add_GetRef(NewTestProcessor<UMassTestProcessor_E>(EntityManager));
 			Proc->GetMutableExecutionOrder().ExecuteInGroup = TEXT("X.Z");
 		}
 
 		{
-			UMassTestProcessorBase* Proc = Processors.Add_GetRef(NewObject<UMassTestProcessor_F>());
+			UMassTestProcessorBase* Proc = Processors.Add_GetRef(NewTestProcessor<UMassTestProcessor_F>(EntityManager));
 			Proc->GetMutableExecutionOrder().ExecuteAfter.Add(UMassTestProcessor_A::StaticClass()->GetFName());
 		}
 		return true;
