[Mass] Added a unit test for all methods of entity creation in MassEntity module

#rnx

diff --git a/Engine/Source/Developer/MassEntityTestSuite/Private/MassEntityTest.cpp b/Engine/Source/Developer/MassEntityTestSuite/Private/MassEntityTest.cpp
--- a/Engine/Source/Developer/MassEntityTestSuite/Private/MassEntityTest.cpp
+++ b/Engine/Source/Developer/MassEntityTestSuite/Private/MassEntityTest.cpp
@@ -1,6 +1,7 @@
 // Copyright Epic Games, Inc. All Rights Reserved.
 
 #include "AITestsCommon.h"
+#include "MassEntityBuilder.h"
 #include "MassEntityManager.h"
 #include "MassProcessingTypes.h"
 #include "MassEntityTestTypes.h"
@@ -407,6 +408,70 @@ struct FEntityTest_SharedPtrFragment : FEntityTestBase
 };
 IMPLEMENT_AI_INSTANT_TEST(FEntityTest_SharedPtrFragment, "System.Mass.Entity.SharedPtrFragment");
 
+struct FEntityTest_CreationPatterns : FEntityTestBase
+{
+	virtual bool InstantTest() override
+	{
+		TArray<FMassEntityHandle> Entities;
+
+		// create entities straight from fragment instances
+		TArray<FInstancedStruct> FragmentInstances = {
+			FInstancedStruct::Make<FTestFragment_Int>()
+			, FInstancedStruct::Make<FTestFragment_Float>()
+		};
+		Entities.Add(EntityManager->CreateEntity(FragmentInstances));
+
+		// create the archetype first, and then create an entity within it.
+		TArray<const UScriptStruct*> FragmentTypes = {
+			FTestFragment_Int::StaticStruct()
+			, FTestFragment_Float::StaticStruct()
+		};
+		FMassArchetypeHandle ArchetypeHandle = EntityManager->CreateArchetype(FragmentTypes);
+		Entities.Add(EntityManager->CreateEntity(ArchetypeHandle));
+
+		// Reserve + Build
+		// the call will result in ReservedEntity being a valid handle, but not associated with any Archetype data
+		Entities.Add(EntityManager->ReserveEntity());
+
+		// any time can pass between reservation and building.
+		// the handle can be used to issue commands
+		//EntityManager->Defer().AddFragment<FTestFragment_Bool>(Entities.Last());
+
+		// the following call will instantiate the entity
+		EntityManager->BuildEntity(Entities.Last(), ArchetypeHandle);
+		// if the commands are flushed before the entity is explicitly built then it will be
+		// done automatically, but subsequent explicit Build calls will fail
+
+		// entity builder, new in 5.6. There's more, see MassEntityBuilderTest.cpp
+		Entities.Add(EntityManager->MakeEntityBuilder()
+			.Add<FTestFragment_Int>(1024)
+			.Add<FTestFragment_Float>(3.14f)
+			.Commit());
+
+		for (FMassEntityHandle EntityHandle : Entities)
+		{
+			const FMassArchetypeHandle ResultArchetypeHandle = EntityManager->GetArchetypeForEntity(EntityHandle);
+			AITEST_TRUE("Each method creates an entity in the same archetype", ResultArchetypeHandle == FloatsIntsArchetype);
+		}
+#if 0
+		// This one requires MassGameplay plugin, here for reference.
+		UMassSpawnerSubsystem* MassSpawnerSubsystem = UWorld::GetSubsystem<UMassSpawnerSubsystem>(GetWorld());
+		FMassEntityTemplateRegistry& TemplateRegistry = MassSpawnerSubsystem->GetMutableTemplateRegistryInstance();
+
+		FMassEntityTemplateData TemplateData;
+		TemplateData.AddFragment<FTestFragment_Int>();
+		TemplateData.AddFragment<FTestFragment_Float>();
+		const TSharedRef<FMassEntityTemplate>& FinalizedTemplate = TemplateRegistry.FindOrAddTemplate(FMassEntityTemplateIDFactory::Make(FGuid::NewGuid()), MoveTemp(TemplateData));
+
+		TArray<FMassEntityHandle> SpawnedEntities;
+		MassSpawnerSubsystem->SpawnEntities(FinalizedTemplate, /*NumberToSpawn=*/1, SpawnedEntities);
+#endif
+
+		return true;
+	}
+};
+IMPLEMENT_AI_INSTANT_TEST(FEntityTest_CreationPatterns, "System.Mass.Entity.CreationPatterns");
+
 } // FMassEntityTestTest
 
 UE_ENABLE_OPTIMIZATION_SHIP
