[Mass] Fixed a bug in group-based archetype sorting in MassEntityQuery that resulted in failed array ensures.

The issue was using a reference to an element of a temp array that was being resized during the loop iteration (thus making the reference unsafe and useless after the reallocation).
Added more complex groups sorting tests to stress-test the change.

#rnx

diff --git a/Engine/Source/Developer/MassEntityTestSuite/Private/MassArchetypeGroupTest.cpp b/Engine/Source/Developer/MassEntityTestSuite/Private/MassArchetypeGroupTest.cpp
--- a/Engine/Source/Developer/MassEntityTestSuite/Private/MassArchetypeGroupTest.cpp
+++ b/Engine/Source/Developer/MassEntityTestSuite/Private/MassArchetypeGroupTest.cpp
@@ -78,14 +78,40 @@ struct FArchetypeGroup_SingleLevelQuery : FEntityTestBase
 };
 IMPLEMENT_AI_INSTANT_TEST(FArchetypeGroup_SingleLevelQuery, "System.Mass.Archetype.Group.SingleLevelQuery");
 
-struct FArchetypeGroup_MultiLevelQuery : FEntityTestBase
+struct FArchetypeGroup_MultiLevelQueryBase : FEntityTestBase
 {
+	TArray<FArchetypeGroupType> GroupTypes;
+	TArray<FMassEntityHandle> Entities;
+	FMassEntityQuery EntityQuery;
+
+	static constexpr int32 NumEntities = 16;
+	TArray<const int32> GroupSizes;
+
+	using FTestFunction = TFunction<bool(FMassExecutionContext&)>;
+	TArray<FTestFunction> Tests;
+
+	virtual bool SetUp() override
+	{
+		if (FEntityTestBase::SetUp())
+		{
+			EntityQuery.Initialize(EntityManager.ToSharedRef());
+			EntityQuery.AddRequirement<FTestFragment_Int>(EMassFragmentAccess::ReadOnly);
+
+			Tests.Add([this](FMassExecutionContext& ExecutionContext)
+			{
+				TArray<TArray<const int32>> ExpectedEntityIndices = {{ 0, 1, 2, 3, 4, 5, 6, 7, 8, 9, 10, 11, 12, 13, 14, 15 }};
+				return TestDefaultOrder(ExecutionContext, ExpectedEntityIndices);
+			});
+
+			return true;
+		}
+		return false;
+	}
+
 	virtual bool InstantTest() override
 	{
-		constexpr int32 GroupsNum = 3;
-		constexpr int32 NumEntities = 16;
-		const int32 GroupSizes[GroupsNum] = { 8, 4, 2 };
-		TArray<FArchetypeGroupType> GroupTypes;
+		const int32 GroupsNum = GroupSizes.Num();
+
 		for (int32 GroupIndex = 0; GroupIndex < GroupsNum; ++GroupIndex)
 		{
 			const FName GroupTypeName(*FString::Printf(TEXT("TestGroup_%d"), GroupIndex));
@@ -94,30 +120,140 @@ struct FArchetypeGroup_MultiLevelQuery : FEntityTestBase
 		
 		AITEST_NOT_EQUAL(TEXT("Two differently named group types are not expected to be equal"), GroupTypes[0], GroupTypes[1]);
 
-		TArray<FMassEntityHandle> Entities;
 		TArray<FMassEntityHandle> VerifyEntities;
 		EntityManager->BatchCreateEntities(IntsArchetype, NumEntities, Entities);
-		const int32 FirstEntityIndex = Entities[0].Index;
 
-		auto GroupAssigner = [&](auto&& Self, const int32 GroupIndex, TArrayView<FMassEntityHandle> EntitiesView) -> void
+		AssignGroups();
+
+		FMassExecutionContext ExecutionContext = EntityManager->CreateExecutionContext(0);
+		
+		int32 TestIndex = 0;
+		bool bResult = true;
+		while(bResult && TestIndex < Tests.Num())
 		{
-			if (GroupIndex < GroupsNum)
+			bResult = Tests[TestIndex++](ExecutionContext);
+		} 
+
+		// @TODO: test sorting when some entities are missing some group types - specifically the first one being used for sorting
+
+		return bResult;
+	}
+
+	virtual void AssignGroups()
+	{
+		auto GroupAssigner = [this](auto&& Self, const int32 GroupIndex, TArrayView<FMassEntityHandle> EntitiesView) -> void
 			{
-				const FArchetypeGroupType& GroupType = GroupTypes[GroupIndex];
-				const int32 GroupSize = GroupSizes[GroupIndex];
-				int32 SubGroupIndex = 0;
-				for (int32 EntityIndex = 0; EntityIndex < EntitiesView.Num(); EntityIndex += GroupSize, ++SubGroupIndex)
+				const int32 GroupsNum = GroupSizes.Num();
+				if (GroupIndex < GroupsNum)
 				{
-					TArrayView<FMassEntityHandle> SubEntities = MakeArrayView(&EntitiesView[EntityIndex], GroupSize);
-					EntityManager->BatchGroupEntities(FArchetypeGroupHandle(GroupType, SubGroupIndex), SubEntities);
-					Self(Self, GroupIndex + 1, SubEntities);
+					const FArchetypeGroupType& GroupType = GroupTypes[GroupIndex];
+					const int32 GroupSize = GroupSizes[GroupIndex];
+					int32 SubGroupIndex = 0;
+					for (int32 EntityIndex = 0; EntityIndex < EntitiesView.Num(); EntityIndex += GroupSize, ++SubGroupIndex)
+					{
+						TArrayView<FMassEntityHandle> SubEntities = MakeArrayView(&EntitiesView[EntityIndex], GroupSize);
+						EntityManager->BatchGroupEntities(FArchetypeGroupHandle(GroupType, SubGroupIndex), SubEntities);
+						Self(Self, GroupIndex + 1, SubEntities);
+					}
 				}
-			}
-		};
+			};
 
 		GroupAssigner(GroupAssigner, 0, Entities);
+	}
 
-		// at this point we have:
+	bool TestDefaultOrder(FMassExecutionContext& ExecutionContext, TArray<TArray<const int32>>& ExpectedEntityIndices)
+	{
+		EntityQuery.ResetGrouping();
+		for (const FArchetypeGroupType& GroupType : GroupTypes)
+		{
+			EntityQuery.GroupBy(GroupType);
+		}
+
+		return TestRaw(ExecutionContext, ExpectedEntityIndices);
+	}
+
+	bool Test1(FMassExecutionContext& ExecutionContext, TArray<TArray<const int32>>& ExpectedEntityIndices)
+	{
+		// now we're going to reverse the sorting of group TestGroup_1
+		EntityQuery.ResetGrouping();
+		EntityQuery.GroupBy(GroupTypes[0]);
+		EntityQuery.GroupBy(GroupTypes[1], DescendingGroupSorter);
+		EntityQuery.GroupBy(GroupTypes[2]);
+
+		return TestRaw(ExecutionContext, ExpectedEntityIndices);
+	}
+
+	bool Test2(FMassExecutionContext& ExecutionContext, TArray<TArray<const int32>>& ExpectedEntityIndices)
+	{
+		// now we're going to reverse the sorting of the last group, TestGroup_2
+		EntityQuery.ResetGrouping();
+		EntityQuery.GroupBy(GroupTypes[0]);
+		EntityQuery.GroupBy(GroupTypes[1]);
+		EntityQuery.GroupBy(GroupTypes[2], DescendingGroupSorter);
+
+		return TestRaw(ExecutionContext, ExpectedEntityIndices);
+	}
+
+	bool Test3(FMassExecutionContext& ExecutionContext, TArray<TArray<const int32>>& ExpectedEntityIndices)
+	{
+		// now we're going to reverse the sorting of the TestGroup_0 and TestGroup_1, and put the latter first
+		EntityQuery.ResetGrouping();
+		EntityQuery.GroupBy(GroupTypes[1], DescendingGroupSorter);
+		EntityQuery.GroupBy(GroupTypes[0], DescendingGroupSorter);
+		EntityQuery.GroupBy(GroupTypes[2]);
+
+		return TestRaw(ExecutionContext, ExpectedEntityIndices);
+	}
+
+	/** Grouping expected to be set before calling this one */
+	bool TestRaw(FMassExecutionContext& ExecutionContext, TArray<TArray<const int32>>& ExpectedEntityIndices)
+	{
+		TArray<FMassEntityHandle> VerifyEntities;
+
+		EntityQuery.ForEachEntityChunk(ExecutionContext, [&VerifyEntities](const FMassExecutionContext& Context)
+		{
+			VerifyEntities.Append(Context.GetEntities());
+		});
+
+		if (ExpectedEntityIndices.Num() == 1)
+		{
+			// we expect the strict order
+			for (int32 ResultIndex = 0; ResultIndex < VerifyEntities.Num(); ++ResultIndex)
+			{
+				AITEST_EQUAL("Expected results vs received", VerifyEntities[ResultIndex], Entities[ExpectedEntityIndices[0][ResultIndex]]);
+			}
+		}
+		else
+		{
+			TArray<FMassEntityHandle> EntitySet;
+			int32 ResultIndex = 0;
+			for (int32 SetIndex = 0; SetIndex < ExpectedEntityIndices.Num(); ++SetIndex)
+			{
+				for (const int32 EntityIndex : ExpectedEntityIndices[SetIndex])
+				{
+					EntitySet.Add(Entities[EntityIndex]);
+				}
+
+				const int32 SetSize = ExpectedEntityIndices[SetIndex].Num();
+				for (int32 SetIterator = 0; SetIterator < SetSize; ++SetIterator, ++ResultIndex)
+				{
+					AITEST_TRUE("Result in expected set", EntitySet.Find(VerifyEntities[ResultIndex]) != INDEX_NONE);
+				}
+				EntitySet.Reset();
+			}
+		}
+
+		return true;
+	}
+};
+
+struct FArchetypeGroup_MultiLevelQuery_A : FArchetypeGroup_MultiLevelQueryBase
+{
+	virtual bool SetUp() override
+	{
+		GroupSizes = { 8, 4, 2 };
+
+		// This will result in the following:
 		// Entity	| TestGroup_0 	| TestGroup_1 	| TestGroup_2
 		//  0		|		0		|		0		|		0
 		//  1		|		0		|		0		|		0
@@ -136,93 +272,221 @@ struct FArchetypeGroup_MultiLevelQuery : FEntityTestBase
 		// 14		|		1		|		1		|		1
 		// 15		|		1		|		1		|		1
 
-		FMassExecutionContext ExecutionContext = EntityManager->CreateExecutionContext(0);
-		FMassEntityQuery EntityQuery(EntityManager);
-		EntityQuery.AddRequirement<FTestFragment_Int>(EMassFragmentAccess::ReadOnly);
-		for (int32 GroupIndex = 0; GroupIndex < GroupsNum; ++GroupIndex)
+		Tests.Add([this](FMassExecutionContext& ExecutionContext)
 		{
-			const FArchetypeGroupType& GroupType = GroupTypes[GroupIndex];
-			EntityQuery.GroupBy(GroupType);
-		}
+			// here's what we expect to see with the TestGroup_1 group sorted in descending order
+			TArray<TArray<const int32>> ExpectedEntityIndices = {{ 4, 5, 6, 7, 0, 1, 2, 3, 12, 13, 14, 15, 8, 9, 10, 11 }};
+			return Test1(ExecutionContext, ExpectedEntityIndices);
+		});
 
-		VerifyEntities.Reset();
-		EntityQuery.ForEachEntityChunk(ExecutionContext, [&VerifyEntities](const FMassExecutionContext& Context)
+		Tests.Add([this](FMassExecutionContext& ExecutionContext)
 		{
-			VerifyEntities.Append(Context.GetEntities());
+			// here's what we expect to see with the TestGroup_1 group sorted in descending order
+			TArray<TArray<const int32>> ExpectedEntityIndices = {{ 2, 3, 0, 1, 6, 7, 4, 5, 10, 11, 8, 9, 14, 15, 12, 13 }};
+			return Test2(ExecutionContext, ExpectedEntityIndices);
 		});
 
-		// by default we don't expect the order of entities processed to change
-		AITEST_TRUE("After grouping with the default sorting the order matches expectations", Algo::Compare(Entities, VerifyEntities));
+		Tests.Add([this](FMassExecutionContext& ExecutionContext)
 		{
-			// now we're going to reverse the sorting of group TestGroup_1
-			EntityQuery.ResetGrouping();
-			EntityQuery.GroupBy(GroupTypes[0]);
-			EntityQuery.GroupBy(GroupTypes[1], DescendingGroupSorter);
-			EntityQuery.GroupBy(GroupTypes[2]);
-
 			// here's what we expect to see with the TestGroup_1 group sorted in descending order
-			const int32 ExpectedEntityIndices[NumEntities] = { 4, 5, 6, 7, 0, 1, 2, 3, 12, 13, 14, 15, 8, 9, 10, 11 };
+			TArray<TArray<const int32>> ExpectedEntityIndices = {{ 12, 13, 14, 15, 4, 5, 6, 7, 8, 9, 10, 11, 0, 1, 2, 3 }};
+			return Test3(ExecutionContext, ExpectedEntityIndices);
+		});
 
-			VerifyEntities.Reset();
-			EntityQuery.ForEachEntityChunk(ExecutionContext, [&VerifyEntities](const FMassExecutionContext& Context)
-			{
-				VerifyEntities.Append(Context.GetEntities());
-			});
+		return FArchetypeGroup_MultiLevelQueryBase::SetUp();
+	}
+};
+IMPLEMENT_AI_INSTANT_TEST(FArchetypeGroup_MultiLevelQuery_A, "System.Mass.Archetype.Group.MultiLevelQuery.A");
 
-			for (int32 ResultIndex = 0; ResultIndex < VerifyEntities.Num(); ++ResultIndex)
-			{
-				AITEST_EQUAL("Sorting TestGroup_1 in a descending order yields the expected results", VerifyEntities[ResultIndex], Entities[ExpectedEntityIndices[ResultIndex]]);
-			}
-		}
+struct FArchetypeGroup_MultiLevelQuery_B : FArchetypeGroup_MultiLevelQueryBase
+{
+	virtual bool SetUp() override
+	{
+		GroupSizes = { 5, 3, 1 };
+
+		// This will result in the following:
+		// Entity	| TestGroup_0 	| TestGroup_1 	| TestGroup_2
+		//  0		|		0		|		0		|		0
+		//  1		|		0		|		0		|		1
+		//  2		|		0		|		0		|		2
+		//  3		|		0		|		1		|		0
+		//  4		|		0		|		1		|		1
+		//  5		|		1		|		0		|		0
+		//  6		|		1		|		0		|		1
+		//  7		|		1		|		0		|		2
+		//  8		|		1		|		1		|		0
+		//  9		|		1		|		1		|		1
+		// 10		|		2		|		0		|		0
+		// 11		|		2		|		0		|		1
+		// 12		|		2		|		0		|		2
+		// 13		|		2		|		1		|		0
+		// 14		|		2		|		1		|		1
+		// 15		|		3		|		0		|		0
+
+		Tests.Add([this](FMassExecutionContext& ExecutionContext)
 		{
-			// now we're going to reverse the sorting of the last group, TestGroup_2
-			EntityQuery.ResetGrouping();
-			EntityQuery.GroupBy(GroupTypes[0]);
-			EntityQuery.GroupBy(GroupTypes[1]);
-			EntityQuery.GroupBy(GroupTypes[2], DescendingGroupSorter);
+			// here's what we expect to see with the TestGroup_1 group sorted in descending order
+			TArray<TArray<const int32>> ExpectedEntityIndices = {{ 3, 4, 0, 1, 2, 8, 9, 5, 6, 7, 13, 14, 10, 11, 12, 15 }};
+			return Test1(ExecutionContext, ExpectedEntityIndices);
+		});
 
+		Tests.Add([this](FMassExecutionContext& ExecutionContext)
+		{
 			// here's what we expect to see with the TestGroup_1 group sorted in descending order
-			const int32 ExpectedEntityIndices[NumEntities] = { 2, 3, 0, 1, 6, 7, 4, 5, 10, 11, 8, 9, 14, 15, 12, 13 };
+			TArray<TArray<const int32>> ExpectedEntityIndices = {{ 2, 1, 0, 4, 3, 7, 6, 5, 9, 8, 12, 11, 10, 14, 13, 15 }};
+			return Test2(ExecutionContext, ExpectedEntityIndices);
+		});
 
-			VerifyEntities.Reset();
-			EntityQuery.ForEachEntityChunk(ExecutionContext, [&VerifyEntities](const FMassExecutionContext& Context)
-			{
-				VerifyEntities.Append(Context.GetEntities());
-			});
+		Tests.Add([this](FMassExecutionContext& ExecutionContext)
+		{
+			// here's what we expect to see with the TestGroup_1 group sorted in descending order
+			TArray<TArray<const int32>> ExpectedEntityIndices = {{ 13, 14, 8, 9, 3, 4, 15, 10, 11, 12, 5, 6, 7, 0, 1, 2 }};
+			return Test3(ExecutionContext, ExpectedEntityIndices);
+		});
 
-			for (int32 ResultIndex = 0; ResultIndex < VerifyEntities.Num(); ++ResultIndex)
+		return FArchetypeGroup_MultiLevelQueryBase::SetUp();
+	}
+};
+IMPLEMENT_AI_INSTANT_TEST(FArchetypeGroup_MultiLevelQuery_B, "System.Mass.Archetype.Group.MultiLevelQuery.B");
+
+struct FArchetypeGroup_MultiLevelQuery_Complex : FArchetypeGroup_MultiLevelQueryBase
+{
+	TArray<TArray<const int32>> OderedSets;
+
+	virtual void AssignGroups() override
+	{
+		const int32 GroupsNum = GroupSizes.Num();
+		for (int32 GroupIndex = 0; GroupIndex < GroupsNum; ++GroupIndex)
+		{
+			const int32 GroupSize = GroupSizes[GroupIndex];
+			int32 SubGroupIndex = 0;
+			for (int32 EntityIndex = 0; EntityIndex < Entities.Num(); EntityIndex += GroupSize)
 			{
-				AITEST_EQUAL("Sorting TestGroup_2 in a descending order yields the expected results", VerifyEntities[ResultIndex], Entities[ExpectedEntityIndices[ResultIndex]]);
+				const int32 SubgroupSize = FMath::Min(Entities.Num() - EntityIndex, GroupSize);
+
+				TArrayView<FMassEntityHandle> SubEntities = MakeArrayView(&Entities[EntityIndex], SubgroupSize);
+				EntityManager->BatchGroupEntities(FArchetypeGroupHandle(GroupTypes[GroupIndex], SubGroupIndex++), SubEntities);
 			}
 		}
+	}
+
+	virtual bool SetUp() override
+	{
+		if (FArchetypeGroup_MultiLevelQueryBase::SetUp() == false)
 		{
-			// now we're going to reverse the sorting of the TestGroup_0 and TestGroup_1, and put the latter first
-			EntityQuery.ResetGrouping();
-			EntityQuery.GroupBy(GroupTypes[1], DescendingGroupSorter);
-			EntityQuery.GroupBy(GroupTypes[0], DescendingGroupSorter);
-			EntityQuery.GroupBy(GroupTypes[2]);
+			return false;
+		}
 
-			// here's what we expect to see with the TestGroup_1 group sorted in descending order
-			const int32 ExpectedEntityIndices[NumEntities] = { 12, 13, 14, 15, 4, 5, 6, 7, 8, 9, 10, 11, 0, 1, 2, 3 };
+		GroupSizes = { 3, 4, 6 };
 
-			VerifyEntities.Reset();
-			EntityQuery.ForEachEntityChunk(ExecutionContext, [&VerifyEntities](const FMassExecutionContext& Context)
-			{
-				VerifyEntities.Append(Context.GetEntities());
-			});
+		// This will result in the following:
+		// Entity	| TestGroup_0 	| TestGroup_1 	| TestGroup_2
+		//  0		|		0		|		0		|		0
+		//  1		|		0		|		0		|		0
+		//  2		|		0		|		0		|		0
+		//  3		|		1		|		0		|		0
+		//  4		|		1		|		1		|		0
+		//  5		|		1		|		1		|		0
+		//  6		|		2		|		1		|		1
+		//  7		|		2		|		1		|		1
+		//  8		|		2		|		2		|		1
+		//  9		|		3		|		2		|		1
+		// 10		|		3		|		2		|		1
+		// 11		|		3		|		2		|		1
+		// 12		|		4		|		3		|		2
+		// 13		|		4		|		3		|		2
+		// 14		|		4		|		3		|		2
+		// 15		|		5		|		3		|		2
+		OderedSets = {
+			/*[0]*/  { 0, 1, 2}
+			/*[1]*/, { 3 }
+			/*[2]*/, { 4, 5 }
+			/*[3]*/, { 6, 7 }
+			/*[4]*/, { 8 }
+			/*[5]*/, { 9, 10, 11 }
+			/*[6]*/, { 12, 13, 14 }
+			/*[7]*/, { 15 }
+		};
 
-			for (int32 ResultIndex = 0; ResultIndex < VerifyEntities.Num(); ++ResultIndex)
+		// dropping the "default order" test since it can produce a slightly different results than expected
+		// due to some entities being in the very same group combinations, like {9, 10, 11} and {12, 13, 14}
+		Tests.Reset();
+		Tests.Add([this](FMassExecutionContext& ExecutionContext)
+		{
+			EntityQuery.ResetGrouping();
+			for (const FArchetypeGroupType& GroupType : GroupTypes)
 			{
-				AITEST_EQUAL("Sorting by TestGroup_2 and TestGroup_1 in a descending order yields the expected results", VerifyEntities[ResultIndex], Entities[ExpectedEntityIndices[ResultIndex]]);
+				EntityQuery.GroupBy(GroupType);
 			}
-		}
 
-		// @TODO: test sorting when some entities are missing some group types - specifically the first one being used for sorting
+			return TestRaw(ExecutionContext, OderedSets);
+		});
+
+		Tests.Add([this](FMassExecutionContext& ExecutionContext)
+		{
+			// here's what we expect to see with the TestGroup_1 group sorted in descending order
+			//TArray<TArray<const int32>> ExpectedEntityIndices = {{ 0, 1, 2, 4, 5, 3, 8, 6, 7, 9, 10, 11, 12, 13, 14, 15 }};
+			TArray<TArray<const int32>> ExpectedSets = {
+				OderedSets[0]
+				, OderedSets[2]
+				, OderedSets[1]
+				, OderedSets[4]
+				, OderedSets[3]
+				, OderedSets[5]
+				, OderedSets[6]
+				, OderedSets[7]
+			};
+			return Test1(ExecutionContext, ExpectedSets);
+		});
+
+		Tests.Add([this](FMassExecutionContext& ExecutionContext)
+		{
+			return Test2(ExecutionContext, OderedSets);
+		});
+
+		Tests.Add([this](FMassExecutionContext& ExecutionContext)
+		{
+			// here's what we expect to see with the TestGroup_1 group sorted in descending order
+			//TArray<TArray<const int32>> ExpectedEntityIndices = {{ 15, 12, 13, 14, 9, 10, 11, 8, 6, 7, 4, 5, 3, 0, 1, 2 }};
+			TArray<TArray<const int32>> ExpectedSets = {
+				OderedSets[7]
+				, OderedSets[6]
+				, OderedSets[5]
+				, OderedSets[4]
+				, OderedSets[3]
+				, OderedSets[2]
+				, OderedSets[1]
+				, OderedSets[0]
+			};
+			return Test3(ExecutionContext, ExpectedSets);
+		});
+
+		Tests.Add([this](FMassExecutionContext& ExecutionContext)
+		{
+			EntityQuery.ResetGrouping();
+			EntityQuery.GroupBy(GroupTypes[2], DescendingGroupSorter);
+			EntityQuery.GroupBy(GroupTypes[0]);
+			EntityQuery.GroupBy(GroupTypes[1], DescendingGroupSorter);
+
+			//TArray<TArray<const int32>> ExpectedEntityIndices = {{ 12, 13, 14, 15, 8, 6, 7, 9, 10, 11, 0, 1, 2, 4, 5, 3 }};
+			TArray<TArray<const int32>> ExpectedSets = {
+				OderedSets[6]
+				, OderedSets[7]
+				, OderedSets[4]
+				, OderedSets[3]
+				, OderedSets[5]
+				, OderedSets[0]
+				, OderedSets[2]
+				, OderedSets[1]
+			};
+
+			return TestRaw(ExecutionContext, ExpectedSets);
+		});
 
 		return true;
 	}
 };
-IMPLEMENT_AI_INSTANT_TEST(FArchetypeGroup_MultiLevelQuery, "System.Mass.Archetype.Group.MultiLevelQuery");
+IMPLEMENT_AI_INSTANT_TEST(FArchetypeGroup_MultiLevelQuery_Complex, "System.Mass.Archetype.Group.MultiLevelQuery.Complex");
 
 /** This test ensures the order in which a given entity is grouped in doesn't matter. */
 struct FArchetypeGroup_ApplicationOrder : FEntityTestBase
diff --git a/Engine/Source/Runtime/MassEntity/Private/MassEntityQuery.cpp b/Engine/Source/Runtime/MassEntity/Private/MassEntityQuery.cpp
--- a/Engine/Source/Runtime/MassEntity/Private/MassEntityQuery.cpp
+++ b/Engine/Source/Runtime/MassEntity/Private/MassEntityQuery.cpp
@@ -602,7 +602,7 @@ void FMassEntityQuery::SortArchetypes(const int32 FirstNewArchetypeIndex)
 		MaxRangeSize = 0;
 		for (; RangesProcessed < RangesThisIteration; ++RangesProcessed)
 		{
-			const TTuple<int32, int32>& Range = Ranges[RangesProcessed];
+			const TTuple<int32, int32> Range = Ranges[RangesProcessed];
 			const int32 RangeLength = Range.Get<1>() - Range.Get<0>();
 			MakeArrayView(&OrderedArchetypeIndices[Range.Get<0>()], RangeLength)
 				.Sort([Step, &GroupIDs, Sorter = GroupSortingSteps[Step].Predicate](const int32 A, const int32 B)
