[Mass] CIS Issue 710523, wrapped DebugGetArchetypeEntitiesCount calls in WITH_MASSENTITY_DEBUG.

#rnx

diff --git a/Engine/Source/Developer/MassEntityTestSuite/Private/MassEntityBuilderTest.cpp b/Engine/Source/Developer/MassEntityTestSuite/Private/MassEntityBuilderTest.cpp
--- a/Engine/Source/Developer/MassEntityTestSuite/Private/MassEntityBuilderTest.cpp
+++ b/Engine/Source/Developer/MassEntityTestSuite/Private/MassEntityBuilderTest.cpp
@@ -481,8 +481,9 @@ namespace UE::Mass::Test::EntityBuilder
 			AITEST_FALSE("(NOT) Run 2: the entity has been created", EntityManager->IsEntityBuilt(ReservedHandle));
 
 			EntityManager->FlushCommands();
-			int32 ArchetypeEntitiesNum = EntityManager->DebugGetArchetypeEntitiesCount(IntsArchetype);
-			AITEST_EQUAL("Number of entities in the target archetype, after flushing", ArchetypeEntitiesNum, 2);
+#if WITH_MASSENTITY_DEBUG
+			AITEST_EQUAL("Number of entities in the target archetype, after flushing", EntityManager->DebugGetArchetypeEntitiesCount(IntsArchetype), 2);
+#endif // WITH_MASSENTITY_DEBUG
 			AITEST_TRUE("the entity has been created", EntityManager->IsEntityBuilt(ReservedHandle));
 
 			return true;
