[Mass] Fixed the outdated comment on UMassTestProcessorBase::ConfigureQueries

#rnx

diff --git a/Engine/Source/Developer/MassEntityTestSuite/Public/MassEntityTestTypes.h b/Engine/Source/Developer/MassEntityTestSuite/Public/MassEntityTestTypes.h
--- a/Engine/Source/Developer/MassEntityTestSuite/Public/MassEntityTestTypes.h
+++ b/Engine/Source/Developer/MassEntityTestSuite/Public/MassEntityTestTypes.h
@@ -201,9 +201,8 @@ public:
 	virtual void Execute(FMassEntityManager& EntityManager, FMassExecutionContext& Context) override;
 	virtual bool ShouldAllowQueryBasedPruning(const bool bRuntimeMode = true) const { return false; }
 	/** 
-	 * ConfigureQueries is called in PostInitProperties, so there's no point in calling anything here, 
-	 * since the EntityQuery configuration will take place after a processor instance gets created
-	 * (so after PostInitProperties has already been called). Use EntityQuery directly instead. 
+	 * Leaving the implementation empty since it's up to the child classes and specific use-cases
+	 * to determine the actual requirements for EntityQuery.
 	 */
 	virtual void ConfigureQueries(const TSharedRef<FMassEntityManager>&) override {}
 
