[Mass] Improved Mass test-time processor for quiter test building.

diff --git a/Engine/Source/Developer/MassEntityTestSuite/Private/MassTestTypes.cpp b/Engine/Source/Developer/MassEntityTestSuite/Private/MassTestTypes.cpp
--- a/Engine/Source/Developer/MassEntityTestSuite/Private/MassTestTypes.cpp
+++ b/Engine/Source/Developer/MassEntityTestSuite/Private/MassTestTypes.cpp
@@ -54,10 +54,35 @@ UMassTestProcessorBase::UMassTestProcessorBase()
 	ExecutionFlags = int32(EProcessorExecutionFlags::All);
 
 	ForEachEntityChunkExecutionFunction = [](FMassExecutionContext& Context) {};
-	ExecutionFunction = [this](FMassEntityManager& InEntityManager, FMassExecutionContext& Context) 
+
+	SetUseParallelForEachEntityChunk(false);
+}
+
+void UMassTestProcessorBase::SetUseParallelForEachEntityChunk(bool bEnable)
+{
+	if (bEnable)
+	{
+		ExecutionFunction = [this](FMassEntityManager& InEntityManager, FMassExecutionContext& Context) 
+		{
+			EntityQuery.ParallelForEachEntityChunk(Context, ForEachEntityChunkExecutionFunction, FMassEntityQuery::EParallelExecutionFlags::Force);
+		};
+	}
+	else
 	{
-		EntityQuery.ForEachEntityChunk(Context, ForEachEntityChunkExecutionFunction);
-	};
+		ExecutionFunction = [this](FMassEntityManager& InEntityManager, FMassExecutionContext& Context) 
+		{
+			EntityQuery.ForEachEntityChunk(Context, ForEachEntityChunkExecutionFunction);
+		};
+	}
+}
+
+void UMassTestProcessorBase::TestExecute(TSharedPtr<FMassEntityManager>& EntityManager)
+{
+	check(EntityManager);
+	UMassTestProcessorBase* This = this;
+	TArrayView<UMassProcessor* const> ProcessorView = MakeArrayView(reinterpret_cast<UMassProcessor**>(&This), 1);
+	FMassProcessingContext ProcessingContext(EntityManager, /*DeltaSeconds=*/0.f);
+	UE::Mass::Executor::RunProcessorsView(ProcessorView, ProcessingContext); 
 }
 
 void UMassTestProcessorBase::Execute(FMassEntityManager& EntityManager, FMassExecutionContext& Context)
diff --git a/Engine/Source/Developer/MassEntityTestSuite/Public/MassEntityTestTypes.h b/Engine/Source/Developer/MassEntityTestSuite/Public/MassEntityTestTypes.h
--- a/Engine/Source/Developer/MassEntityTestSuite/Public/MassEntityTestTypes.h
+++ b/Engine/Source/Developer/MassEntityTestSuite/Public/MassEntityTestTypes.h
@@ -215,9 +215,11 @@ public:
 	 */
 	FMassExecuteFunction ForEachEntityChunkExecutionFunction;
 
-
 	void SetShouldAllowMultipleInstances(const bool bInShouldAllowDuplicated) { bAllowMultipleInstances = bInShouldAllowDuplicated; }
 
+	void SetUseParallelForEachEntityChunk(bool bEnable);
+	void TestExecute(TSharedPtr<FMassEntityManager>& EntityManager);
+
 	/** public on purpose, this is a test processor, no worries about access*/
 	FMassEntityQuery EntityQuery;
 };
