[Mass]  Issue 724599: Compile errors in MassEntityIteratorTest.cpp in cooked editor

#rnx

diff --git a/Engine/Source/Developer/MassEntityTestSuite/Private/MassEntityIteratorTest.cpp b/Engine/Source/Developer/MassEntityTestSuite/Private/MassEntityIteratorTest.cpp
--- a/Engine/Source/Developer/MassEntityTestSuite/Private/MassEntityIteratorTest.cpp
+++ b/Engine/Source/Developer/MassEntityTestSuite/Private/MassEntityIteratorTest.cpp
@@ -10,6 +10,7 @@
 #define LOCTEXT_NAMESPACE "MassTest"
 
 UE_DISABLE_OPTIMIZATION_SHIP
+#if WITH_MASSENTITY_DEBUG
 
 namespace UE::Mass::Test::EntityIterator
 {
@@ -138,6 +139,7 @@ IMPLEMENT_AI_INSTANT_TEST(FIterator_QuerylessIterator, "System.Mass.Entity.Itera
 
 } // UE::Mass::Test
 
+#endif // WITH_MASSENTITY_DEBUG
 UE_ENABLE_OPTIMIZATION_SHIP
 
 #undef LOCTEXT_NAMESPACE
