[Mass] fixed a bug in Mass group unit tests that resulted in heap overflow issues as caught by AddressSanitized

#rnx

diff --git a/Engine/Source/Developer/MassEntityTestSuite/Private/MassArchetypeGroupTest.cpp b/Engine/Source/Developer/MassEntityTestSuite/Private/MassArchetypeGroupTest.cpp
--- a/Engine/Source/Developer/MassEntityTestSuite/Private/MassArchetypeGroupTest.cpp
+++ b/Engine/Source/Developer/MassEntityTestSuite/Private/MassArchetypeGroupTest.cpp
@@ -151,7 +151,8 @@ struct FArchetypeGroup_MultiLevelQueryBase : FEntityTestBase
 					int32 SubGroupIndex = 0;
 					for (int32 EntityIndex = 0; EntityIndex < EntitiesView.Num(); EntityIndex += GroupSize, ++SubGroupIndex)
 					{
-						TArrayView<FMassEntityHandle> SubEntities = MakeArrayView(&EntitiesView[EntityIndex], GroupSize);
+						const int32 ThisGroupSize = FMath::Min(EntitiesView.Num() - EntityIndex, GroupSize);
+						TArrayView<FMassEntityHandle> SubEntities = MakeArrayView(&EntitiesView[EntityIndex], ThisGroupSize);
 						EntityManager->BatchGroupEntities(FArchetypeGroupHandle(GroupType, SubGroupIndex), SubEntities);
 						Self(Self, GroupIndex + 1, SubEntities);
 					}
