[Mass] tiny improvement to Mass storage unit testing to ensure some of the tests finish in a fixed amount of steps

diff --git a/Engine/Source/Developer/MassEntityTestSuite/Private/MassEntityStorageTest.cpp b/Engine/Source/Developer/MassEntityTestSuite/Private/MassEntityStorageTest.cpp
--- a/Engine/Source/Developer/MassEntityTestSuite/Private/MassEntityStorageTest.cpp
+++ b/Engine/Source/Developer/MassEntityTestSuite/Private/MassEntityStorageTest.cpp
@@ -21,10 +21,21 @@ enum class EEntityManagerMode : int32
 
 struct FEntityStorageTestBase : FAITestBase
 {
+	enum EEntityOperation
+	{
+		Operation_IndividualReserve,
+		Operation_BatchReserve,
+		Operation_IndividualDestroy,
+		Operation_BatchDestroy,
+
+		OperationsCount,
+	};
+
 	TArray<TSharedPtr<FMassEntityManager>> EntityManagers;
 	TArray<EEntityManagerMode> ManagersToCreate;
 	const uint32 MaxConcurrentEntitiesPerPage = 0;
 	const int32 TotalNumToReserve = 0;
+	const int32 OperationsNumLimit = OperationsCount * 5;
 	int32 NumToReserveInOneIteration = 0;
 	int32 NumToReleaseInOneIteration = 0;
 	FRandomStream RandomStream;
@@ -42,14 +53,6 @@ struct FEntityStorageTestBase : FAITestBase
 		NumToReleaseInOneIteration = TotalNumToReserve / 12;
 	}
 
-	enum EEntityOperation
-	{
-		Operation_IndividualReserve,
-		Operation_BatchReserve,
-		Operation_IndividualDestroy,
-		Operation_BatchDestroy,
-	};
-
 	void PerformOperation(const EEntityOperation CurrentOperation, TSharedPtr<FMassEntityManager>& EntityManager, TArray<FMassEntityHandle>& EntitiesReserved)
 	{
 		switch (CurrentOperation)
@@ -341,7 +344,7 @@ struct FEntityStorageTest_MixedOperationsBase : FEntityStorageTestBase
 		TArray<FMassEntityHandle> EntitiesReserved;
 
 		EEntityOperation CurrentOperation = Operation_BatchReserve;
-		while (EntitiesReserved.Num() < TotalNumToReserve)
+		while (EntitiesReserved.Num() < TotalNumToReserve && OperationsPerformed.Num() < OperationsNumLimit)
 		{
 			PerformOperation(CurrentOperation, EntityManager, EntitiesReserved);
 			const bool bValidColleciton = ValidateUniqueAndValidEntities(EntityManager, EntitiesReserved);
