[Mass] Removed some unnecessary headers inclusion

#rnx

diff --git a/Engine/Source/Developer/MassEntityTestSuite/Private/MassUtilsTest.cpp b/Engine/Source/Developer/MassEntityTestSuite/Private/MassUtilsTest.cpp
--- a/Engine/Source/Developer/MassEntityTestSuite/Private/MassUtilsTest.cpp
+++ b/Engine/Source/Developer/MassEntityTestSuite/Private/MassUtilsTest.cpp
@@ -1,12 +1,8 @@
 // Copyright Epic Games, Inc. All Rights Reserved.
 
 #include "MassEntityManager.h"
-#include "MassProcessingTypes.h"
 #include "MassEntityTestTypes.h"
 #include "MassEntityTypes.h"
-#include "MassEntityView.h"
-
-#include "Algo/Sort.h"
 
 #define LOCTEXT_NAMESPACE "MassTest"
 
diff --git a/Engine/Source/Runtime/MassEntity/Private/MassProcessingTypes.cpp b/Engine/Source/Runtime/MassEntity/Private/MassProcessingTypes.cpp
--- a/Engine/Source/Runtime/MassEntity/Private/MassProcessingTypes.cpp
+++ b/Engine/Source/Runtime/MassEntity/Private/MassProcessingTypes.cpp
@@ -2,7 +2,6 @@
 
 #include "MassProcessingTypes.h"
 #include "MassProcessor.h"
-#include "Misc/OutputDevice.h"
 #include "MassEntityUtils.h"
 #include "VisualLogger/VisualLogger.h"
 #include "MassDebugger.h"
