[Mass] Added a couple of observer tests

The CL also moves CreationContext tests into a separate file.

#rnx

diff --git a/Engine/Source/Developer/MassEntityTestSuite/Private/MassCreationContextTest.cpp b/Engine/Source/Developer/MassEntityTestSuite/Private/MassCreationContextTest.cpp
new file mode 100644
--- /dev/null
+++ b/Engine/Source/Developer/MassEntityTestSuite/Private/MassCreationContextTest.cpp
@@ -0,0 +1,128 @@
+// Copyright Epic Games, Inc. All Rights Reserved.
+
+#include "MassEntityManager.h"
+#include "MassProcessingTypes.h"
+#include "MassEntityTestTypes.h"
+#include "MassEntityTypes.h"
+#include "MassCommandBuffer.h"
+#include "MassExecutionContext.h"
+
+#define LOCTEXT_NAMESPACE "MassTest"
+
+UE_DISABLE_OPTIMIZATION_SHIP
+
+namespace UE::Mass::Test::CreationContext
+{
+//-----------------------------------------------------------------------------
+// creation context 
+//-----------------------------------------------------------------------------
+struct FCreationContextTest : FEntityTestBase
+{
+	virtual bool InstantTest() override
+	{
+		constexpr int32 IntEntitiesToSpawnCount = 6;
+		constexpr int32 FloatEntitiesToSpawnCount = 7;
+
+		TArray<FMassEntityHandle> Entities;
+		TSharedRef<FMassEntityManager::FEntityCreationContext> CreationContextInt = EntityManager->BatchCreateEntities(IntsArchetype, IntEntitiesToSpawnCount, Entities);
+		TSharedRef<FMassEntityManager::FEntityCreationContext> CreationContextFloat = EntityManager->BatchCreateEntities(FloatsArchetype, FloatEntitiesToSpawnCount, Entities);
+		const int32 NumDifferentArchetypesUsed = 2;
+
+		AITEST_EQUAL(TEXT("Two back to back entity creation operations should result in the same creation context"), CreationContextInt, CreationContextFloat);
+		AITEST_FALSE(TEXT("CreationContext's entity collection should be invalidated at this moment"), CreationContextInt->DebugAreEntityCollectionsUpToDate());
+
+		TConstArrayView<FMassArchetypeEntityCollection> EntityCollections = CreationContextInt->GetEntityCollections();
+		AITEST_EQUAL(TEXT("We expect the number of resulting collections to match expectations"), EntityCollections.Num(), NumDifferentArchetypesUsed);
+
+		return true;
+	}
+};
+IMPLEMENT_AI_INSTANT_TEST(FCreationContextTest, "System.Mass.CreationContext.Append");
+
+struct FManualCreate : FEntityTestBase
+{
+	virtual bool InstantTest() override
+	{
+		constexpr int32 IntEntitiesToSpawnCount = 6;
+		constexpr int32 FloatEntitiesToSpawnCount = 7;
+		int NumDifferentArchetypesUsed = 0;
+
+		TArray<FMassEntityHandle> Entities;
+		TSharedRef<FMassEntityManager::FEntityCreationContext> ObtainedContext = EntityManager->GetOrMakeCreationContext();
+		{
+			TSharedRef<FMassEntityManager::FEntityCreationContext> ObtainedContextCopy = EntityManager->GetOrMakeCreationContext();
+			AITEST_EQUAL(TEXT("Two back to back creation context fetching should result in the same instance"), ObtainedContext, ObtainedContextCopy);
+		}
+
+		{
+			TSharedRef<FMassEntityManager::FEntityCreationContext> CreationContextInt = EntityManager->BatchCreateEntities(IntsArchetype, IntEntitiesToSpawnCount, Entities);
+			AITEST_EQUAL(TEXT("Creating entities should return the original context"), ObtainedContext, CreationContextInt);
+			++NumDifferentArchetypesUsed;
+		}
+		
+		AITEST_TRUE(TEXT("CreationContext's entity collection should be still valid at this moment since we only added one entity collection/array")
+			, ObtainedContext->DebugAreEntityCollectionsUpToDate());
+
+		{
+			TSharedRef<FMassEntityManager::FEntityCreationContext> TempContext = EntityManager->BatchCreateEntities(IntsArchetype, IntEntitiesToSpawnCount, Entities);
+			AITEST_EQUAL(TEXT("Creating entities should return the original context"), ObtainedContext, TempContext);
+
+			AITEST_FALSE(TEXT("CreationContext's entity collection should be invalidated at this moment")
+				, TempContext->DebugAreEntityCollectionsUpToDate());
+		}
+
+		TConstArrayView<FMassArchetypeEntityCollection> EntityCollections = ObtainedContext->GetEntityCollections();
+		AITEST_EQUAL(TEXT("We expect the number of resulting collections to match expectations"), EntityCollections.Num(), NumDifferentArchetypesUsed);
+
+		return true;
+	}
+};
+IMPLEMENT_AI_INSTANT_TEST(FManualCreate, "System.Mass.CreationContext.ManualCreate");
+
+struct FManualBuild : FEntityTestBase
+{
+	virtual bool InstantTest() override
+	{
+		constexpr int32 FloatEntitiesToSpawnCount = 7;
+		int NumDifferentArchetypesUsed = 0;
+
+		TArray<FTestFragment_Float> Payload;
+		for (int Index = 0; Index < FloatEntitiesToSpawnCount; ++Index)
+		{ 
+			Payload.Add(FTestFragment_Float(float(Index)));
+		}
+
+		TSharedRef<FMassEntityManager::FEntityCreationContext> ObtainedContext = EntityManager->GetOrMakeCreationContext();
+		
+		TArray<FMassEntityHandle> Entities;
+		EntityManager->BatchReserveEntities(FloatEntitiesToSpawnCount, Entities);
+
+		FStructArrayView PaloadView(Payload);
+		TArray<FMassArchetypeEntityCollectionWithPayload> EntityCollections;
+		FMassArchetypeEntityCollectionWithPayload::CreateEntityRangesWithPayload(*EntityManager, Entities, FMassArchetypeEntityCollection::NoDuplicates
+			, FMassGenericPayloadView(MakeArrayView(&PaloadView, 1)), EntityCollections);
+
+		checkf(EntityCollections.Num() <= 1, TEXT("We expect TargetEntities to only contain archetype-less entities, ones that need to be \'build\'"));
+
+		{
+			TSharedRef<FMassEntityManager::FEntityCreationContext> CreationContext = EntityManager->BatchBuildEntities(EntityCollections[0], FMassFragmentBitSet(*FTestFragment_Float::StaticStruct()));
+			AITEST_EQUAL(TEXT("Creating entities should return the original context"), ObtainedContext, CreationContext);
+			++NumDifferentArchetypesUsed;
+		}
+
+		AITEST_TRUE(TEXT("CreationContext's entity collection should be still valid at this moment since we only added one entity collection/array")
+			, ObtainedContext->DebugAreEntityCollectionsUpToDate());
+
+		TConstArrayView<FMassArchetypeEntityCollection> ContextEntityCollections = ObtainedContext->GetEntityCollections();
+		AITEST_EQUAL(TEXT("We expect the number of resulting collections to match expectations"), ContextEntityCollections.Num(), NumDifferentArchetypesUsed);
+
+		return true;
+	}
+};
+IMPLEMENT_AI_INSTANT_TEST(FManualBuild, "System.Mass.CreationContext.ManualBuild");
+
+} // UE::Mass::Test::CreationContext
+
+UE_ENABLE_OPTIMIZATION_SHIP
+
+#undef LOCTEXT_NAMESPACE
diff --git a/Engine/Source/Developer/MassEntityTestSuite/Private/MassObserverTest.cpp b/Engine/Source/Developer/MassEntityTestSuite/Private/MassObserverTest.cpp
--- a/Engine/Source/Developer/MassEntityTestSuite/Private/MassObserverTest.cpp
+++ b/Engine/Source/Developer/MassEntityTestSuite/Private/MassObserverTest.cpp
@@ -14,7 +14,7 @@ UE_DISABLE_OPTIMIZATION_SHIP
 //----------------------------------------------------------------------//
 // tests 
 //----------------------------------------------------------------------//
-namespace FMassObserverTest
+namespace UE::Mass::Test::Observers
 {
 
 auto EntityIndexSorted = [](const FMassEntityHandle& A, const FMassEntityHandle& B)
@@ -86,9 +86,9 @@ struct FTagBaseOperation : FEntityTestBase
 	}
 };
 
-struct FObserverProcessorTest_SingleEntitySingleArchetypeAdd : FTagBaseOperation
+struct FSingleEntitySingleArchetypeAdd : FTagBaseOperation
 {
-	FObserverProcessorTest_SingleEntitySingleArchetypeAdd() { OperationObserved = EMassObservedOperation::Add; }
+	FSingleEntitySingleArchetypeAdd() { OperationObserved = EMassObservedOperation::Add; }
 	virtual bool PerformOperation() override 
 	{
 		ExpectedEntities = { EntitiesInt[1] };
@@ -96,11 +96,11 @@ struct FObserverProcessorTest_SingleEntitySingleArchetypeAdd : FTagBaseOperation
 		return true;
 	}
 };
-IMPLEMENT_AI_INSTANT_TEST(FObserverProcessorTest_SingleEntitySingleArchetypeAdd, "System.Mass.Observer.Tag.SingleEntitySingleArchetypeAdd");
+IMPLEMENT_AI_INSTANT_TEST(FSingleEntitySingleArchetypeAdd, "System.Mass.Observer.Tag.SingleEntitySingleArchetypeAdd");
 
-struct FObserverProcessorTest_SingleEntitySingleArchetypeRemove : FTagBaseOperation
+struct FSingleEntitySingleArchetypeRemove : FTagBaseOperation
 {
-	FObserverProcessorTest_SingleEntitySingleArchetypeRemove() { OperationObserved = EMassObservedOperation::Remove; }
+	FSingleEntitySingleArchetypeRemove() { OperationObserved = EMassObservedOperation::Remove; }
 	virtual bool PerformOperation() override
 	{
 		ExpectedEntities = { EntitiesInt[1] };
@@ -113,27 +113,28 @@ struct FObserverProcessorTest_SingleEntitySingleArchetypeRemove : FTagBaseOperat
 		return true;
 	}
 };
-IMPLEMENT_AI_INSTANT_TEST(FObserverProcessorTest_SingleEntitySingleArchetypeRemove, "System.Mass.Observer.Tag.SingleEntitySingleArchetypeRemove");
+IMPLEMENT_AI_INSTANT_TEST(FSingleEntitySingleArchetypeRemove, "System.Mass.Observer.Tag.SingleEntitySingleArchetypeRemove");
 
-struct FObserverProcessorTest_SingleEntitySingleArchetypeDestroy : FTagBaseOperation
+struct FSingleEntitySingleArchetypeDestroy : FTagBaseOperation
 {
-	FObserverProcessorTest_SingleEntitySingleArchetypeDestroy() { OperationObserved = EMassObservedOperation::Remove; }
+	FSingleEntitySingleArchetypeDestroy() { OperationObserved = EMassObservedOperation::Remove; }
 	virtual bool PerformOperation() override
 	{
 		ExpectedEntities = { EntitiesInt[1] };
 		EntityManager->Defer().AddTag<FTagStruct>(EntitiesInt[1]);
 		EntityManager->FlushCommands();
 		// since we're only observing tag removal we don't expect AffectedEntities to contain any data at this point
-		AITEST_EQUAL(TEXT("Tag addition is not being observed and is not expected to produce results yet"), AffectedEntities.Num(), 0);
+		AITEST_EQUAL(TEXT("FTagStruct addition is not being observed and is not expected to produce results yet"), AffectedEntities.Num(), 0);
 		EntityManager->Defer().DestroyEntity(EntitiesInt[1]);
+		
 		return true;
 	}
 };
-IMPLEMENT_AI_INSTANT_TEST(FObserverProcessorTest_SingleEntitySingleArchetypeDestroy, "System.Mass.Observer.Tag.SingleEntitySingleArchetypeDestroy");
+IMPLEMENT_AI_INSTANT_TEST(FSingleEntitySingleArchetypeDestroy, "System.Mass.Observer.Tag.SingleEntitySingleArchetypeDestroy");
 
-struct FObserverProcessorTest_MultipleArchetypeAdd : FTagBaseOperation
+struct FMultipleArchetypeAdd : FTagBaseOperation
 {
-	FObserverProcessorTest_MultipleArchetypeAdd() { OperationObserved = EMassObservedOperation::Add; }
+	FMultipleArchetypeAdd() { OperationObserved = EMassObservedOperation::Add; }
 
 	virtual bool PerformOperation() override
 	{
@@ -145,11 +146,28 @@ struct FObserverProcessorTest_MultipleArchetypeAdd : FTagBaseOperation
 		return true;
 	}
 };
-IMPLEMENT_AI_INSTANT_TEST(FObserverProcessorTest_MultipleArchetypeAdd, "System.Mass.Observer.Tag.MultipleArchetypesAdd");
+IMPLEMENT_AI_INSTANT_TEST(FMultipleArchetypeAdd, "System.Mass.Observer.Tag.MultipleArchetypesAdd");
+
 
-struct FObserverProcessorTest_MultipleArchetypeRemove : FTagBaseOperation
+struct FMultipleArchetypeAdd_Sync : FTagBaseOperation
 {
-	FObserverProcessorTest_MultipleArchetypeRemove() { OperationObserved = EMassObservedOperation::Remove; }
+	FMultipleArchetypeAdd_Sync() { OperationObserved = EMassObservedOperation::Add; }
+
+	virtual bool PerformOperation() override
+	{
+		ExpectedEntities = { EntitiesInt[0], EntitiesInt[2], EntitiesIntsFloat[1] };
+		for (const FMassEntityHandle& ModifiedEntity : ExpectedEntities)
+		{
+			EntityManager->AddTagToEntity(ModifiedEntity, FTagStruct::StaticStruct());
+		}
+		return true;
+	}
+};
+IMPLEMENT_AI_INSTANT_TEST(FMultipleArchetypeAdd_Sync, "System.Mass.Observer.Tag.MultipleArchetypesAdd_Sync");
+
+struct FMultipleArchetypeRemove : FTagBaseOperation
+{
+	FMultipleArchetypeRemove() { OperationObserved = EMassObservedOperation::Remove; }
 
 	virtual bool PerformOperation() override
 	{
@@ -160,19 +178,44 @@ struct FObserverProcessorTest_MultipleArchetypeRemove : FTagBaseOperation
 		}
 		EntityManager->FlushCommands();
 		// since we're only observing tag removal we don't expect AffectedEntities to contain any data at this point
-		AITEST_EQUAL(TEXT("Tag addition is not being observed and is not expected to produce results yet"), AffectedEntities.Num(), 0);
+		AITEST_EQUAL(TEXT("FTagStruct addition is not being observed and is not expected to produce results yet"), AffectedEntities.Num(), 0);
 		for (const FMassEntityHandle& ModifiedEntity : ExpectedEntities)
 		{
 			EntityManager->Defer().RemoveTag<FTagStruct>(ModifiedEntity);
 		}
+
 		return true;
 	}
 };
-IMPLEMENT_AI_INSTANT_TEST(FObserverProcessorTest_MultipleArchetypeRemove, "System.Mass.Observer.Tag.MultipleArchetypesRemove");
+IMPLEMENT_AI_INSTANT_TEST(FMultipleArchetypeRemove, "System.Mass.Observer.Tag.MultipleArchetypesRemove");
 
-struct FObserverProcessorTest_MultipleArchetypeDestroy : FTagBaseOperation
+struct FMultipleArchetypeRemove_Sync : FTagBaseOperation
 {
-	FObserverProcessorTest_MultipleArchetypeDestroy() { OperationObserved = EMassObservedOperation::Remove; }
+	FMultipleArchetypeRemove_Sync() { OperationObserved = EMassObservedOperation::Remove; }
+
+	virtual bool PerformOperation() override
+	{
+		ExpectedEntities = { EntitiesInt[0], EntitiesInt[2], EntitiesIntsFloat[1] };
+		for (const FMassEntityHandle& ModifiedEntity : ExpectedEntities)
+		{
+			EntityManager->AddTagToEntity(ModifiedEntity, FTagStruct::StaticStruct());
+		}
+		
+		// since we're only observing tag removal we don't expect AffectedEntities to contain any data at this point
+		AITEST_EQUAL(TEXT("FTagStruct addition is not being observed and is not expected to produce results yet"), AffectedEntities.Num(), 0);
+		for (const FMassEntityHandle& ModifiedEntity : ExpectedEntities)
+		{
+			EntityManager->RemoveTagFromEntity(ModifiedEntity, FTagStruct::StaticStruct());
+		}
+
+		return true;
+	}
+};
+IMPLEMENT_AI_INSTANT_TEST(FMultipleArchetypeRemove_Sync, "System.Mass.Observer.Tag.MultipleArchetypesRemove_Sync");
+
+struct FMultipleArchetypeDestroy : FTagBaseOperation
+{
+	FMultipleArchetypeDestroy() { OperationObserved = EMassObservedOperation::Remove; }
 
 	virtual bool PerformOperation() override
 	{
@@ -191,11 +234,65 @@ struct FObserverProcessorTest_MultipleArchetypeDestroy : FTagBaseOperation
 		return true;
 	}
 };
-IMPLEMENT_AI_INSTANT_TEST(FObserverProcessorTest_MultipleArchetypeDestroy, "System.Mass.Observer.Tag.MultipleArchetypesDestroy");
+IMPLEMENT_AI_INSTANT_TEST(FMultipleArchetypeDestroy, "System.Mass.Observer.Tag.MultipleArchetypesDestroy");
 
-struct FObserverProcessorTest_MultipleArchetypeSwap : FTagBaseOperation
+struct FForbidModifyOnDestroy : FTagBaseOperation
 {
-	FObserverProcessorTest_MultipleArchetypeSwap() { OperationObserved = EMassObservedOperation::Remove; }
+	using Super = FTagBaseOperation;
+	FForbidModifyOnDestroy() { OperationObserved = EMassObservedOperation::Remove; }
+
+	virtual bool SetUp() override
+	{
+		if (Super::SetUp())
+		{
+			ObserverProcessor->ForEachEntityChunkExecutionFunction = [bCommandsFlushedPtr = &bCommandsFlushed, AffectedEntitiesPtr = &AffectedEntities](FMassExecutionContext& Context)
+				{
+					AffectedEntitiesPtr->Append(Context.GetEntities().GetData(), Context.GetEntities().Num());
+
+					// try changing the input entities' composition.
+					const bool bIsProcessing = Context.GetEntityManagerChecked().IsProcessing();
+					//for (int32 EntityIndex = 0; EntityIndex < Context.GetEntities().Num(); ++EntityIndex)
+					for (const FMassEntityHandle EntityHandle : Context.GetEntities())
+					{
+						//Context.Defer().AddTag<FTestTag_A>(EntityHandle);
+						//Context.GetEntityManagerChecked().AddTagToEntity(EntityHandle, FTestTag_B::StaticStruct());
+					}
+
+					Context.Defer().PushCommand<FMassDeferredSetCommand>([&bCommandsFlushedPtr](FMassEntityManager&)
+						{
+							// dummy command, here just to catch if commands issue by observers got executed at all
+							*bCommandsFlushedPtr = true;
+						});
+				};
+
+			return true;
+		}
+		return false;
+	}
+
+	virtual bool PerformOperation() override
+	{
+		ExpectedEntities = { EntitiesInt[0], EntitiesInt[2], EntitiesIntsFloat[1] };
+		for (const FMassEntityHandle& ModifiedEntity : ExpectedEntities)
+		{
+			EntityManager->Defer().AddTag<FTagStruct>(ModifiedEntity);
+		}
+		EntityManager->FlushCommands();
+
+		// since we're only observing tag removal we don't expect AffectedEntities to contain any data at this point
+		AITEST_EQUAL(TEXT("Tag addition is not being observed and is not expected to produce results yet"), AffectedEntities.Num(), 0);
+		for (const FMassEntityHandle& ModifiedEntity : ExpectedEntities)
+		{
+			EntityManager->Defer().DestroyEntity(ModifiedEntity);
+		}
+		return true;
+	}
+};
+IMPLEMENT_AI_INSTANT_TEST(FForbidModifyOnDestroy, "System.Mass.Observer.ForbidModifyOnDestroy");
+
+struct FMultipleArchetypeSwap : FTagBaseOperation
+{
+	FMultipleArchetypeSwap() { OperationObserved = EMassObservedOperation::Remove; }
 
 	virtual bool PerformOperation() override
 	{
@@ -214,11 +311,11 @@ struct FObserverProcessorTest_MultipleArchetypeSwap : FTagBaseOperation
 		return true;
 	}
 };
-IMPLEMENT_AI_INSTANT_TEST(FObserverProcessorTest_MultipleArchetypeSwap, "System.Mass.Observer.Tag.MultipleArchetypesSwap");
+IMPLEMENT_AI_INSTANT_TEST(FMultipleArchetypeSwap, "System.Mass.Observer.Tag.MultipleArchetypesSwap");
 
-struct FObserverProcessorTest_EntityCreation_Individuals : FTagBaseOperation
+struct FEntityCreation_Individuals : FTagBaseOperation
 {
-	FObserverProcessorTest_EntityCreation_Individuals() { OperationObserved = EMassObservedOperation::Add; }
+	FEntityCreation_Individuals() { OperationObserved = EMassObservedOperation::Add; }
 
 	virtual bool InstantTest() override
 	{
@@ -243,11 +340,11 @@ struct FObserverProcessorTest_EntityCreation_Individuals : FTagBaseOperation
 		return true;
 	}
 };
-IMPLEMENT_AI_INSTANT_TEST(FObserverProcessorTest_EntityCreation_Individuals, "System.Mass.Observer.Create.TagInvididualEntities");
+IMPLEMENT_AI_INSTANT_TEST(FEntityCreation_Individuals, "System.Mass.Observer.Create.TagInvididualEntities");
 
-struct FObserverProcessorTest_EntityCreation_Batched : FTagBaseOperation
+struct FEntityCreation_Batched : FTagBaseOperation
 {
-	FObserverProcessorTest_EntityCreation_Batched() { OperationObserved = EMassObservedOperation::Add; }
+	FEntityCreation_Batched() { OperationObserved = EMassObservedOperation::Add; }
 
 	virtual bool InstantTest() override
 	{
@@ -272,7 +369,7 @@ struct FObserverProcessorTest_EntityCreation_Batched : FTagBaseOperation
 		return true;
 	}
 };
-IMPLEMENT_AI_INSTANT_TEST(FObserverProcessorTest_EntityCreation_Batched, "System.Mass.Observer.Create.TagBatchedEntities");
+IMPLEMENT_AI_INSTANT_TEST(FEntityCreation_Batched, "System.Mass.Observer.Create.TagBatchedEntities");
 
 //-----------------------------------------------------------------------------
 // fragments
@@ -340,9 +437,9 @@ struct FFragmentTestBase : FEntityTestBase
 	}
 };
 
-struct FFragmentObserverTest_SingleEntitySingleArchetypeAdd : FFragmentTestBase
+struct FFragmentTest_SingleEntitySingleArchetypeAdd : FFragmentTestBase
 {
-	FFragmentObserverTest_SingleEntitySingleArchetypeAdd() { OperationObserved = EMassObservedOperation::Add; }
+	FFragmentTest_SingleEntitySingleArchetypeAdd() { OperationObserved = EMassObservedOperation::Add; }
 	virtual bool PerformOperation() override
 	{
 		ExpectedEntities = { EntitiesInt[1] };
@@ -350,11 +447,11 @@ struct FFragmentObserverTest_SingleEntitySingleArchetypeAdd : FFragmentTestBase
 		return true;
 	}
 };
-IMPLEMENT_AI_INSTANT_TEST(FFragmentObserverTest_SingleEntitySingleArchetypeAdd, "System.Mass.Observer.Fragment.SingleEntitySingleArchetypeAdd");
+IMPLEMENT_AI_INSTANT_TEST(FFragmentTest_SingleEntitySingleArchetypeAdd, "System.Mass.Observer.Fragment.SingleEntitySingleArchetypeAdd");
 
-struct FFragmentObserverTest_SingleEntitySingleArchetypeRemove : FFragmentTestBase
+struct FFragmentTest_SingleEntitySingleArchetypeRemove : FFragmentTestBase
 {
-	FFragmentObserverTest_SingleEntitySingleArchetypeRemove() { OperationObserved = EMassObservedOperation::Remove; }
+	FFragmentTest_SingleEntitySingleArchetypeRemove() { OperationObserved = EMassObservedOperation::Remove; }
 	virtual bool PerformOperation() override
 	{
 		ExpectedEntities = { EntitiesInt[1] };
@@ -367,11 +464,11 @@ struct FFragmentObserverTest_SingleEntitySingleArchetypeRemove : FFragmentTestBa
 		return true;
 	}
 };
-IMPLEMENT_AI_INSTANT_TEST(FFragmentObserverTest_SingleEntitySingleArchetypeRemove, "System.Mass.Observer.Fragment.SingleEntitySingleArchetypeRemove");
+IMPLEMENT_AI_INSTANT_TEST(FFragmentTest_SingleEntitySingleArchetypeRemove, "System.Mass.Observer.Fragment.SingleEntitySingleArchetypeRemove");
 
-struct FFragmentObserverTest_SingleEntitySingleArchetypeDestroy : FFragmentTestBase
+struct FFragmentTest_SingleEntitySingleArchetypeDestroy : FFragmentTestBase
 {
-	FFragmentObserverTest_SingleEntitySingleArchetypeDestroy() { OperationObserved = EMassObservedOperation::Remove; }
+	FFragmentTest_SingleEntitySingleArchetypeDestroy() { OperationObserved = EMassObservedOperation::Remove; }
 	virtual bool PerformOperation() override
 	{
 		ExpectedEntities = { EntitiesInt[1] };
@@ -383,11 +480,11 @@ struct FFragmentObserverTest_SingleEntitySingleArchetypeDestroy : FFragmentTestB
 		return true;
 	}
 };
-IMPLEMENT_AI_INSTANT_TEST(FFragmentObserverTest_SingleEntitySingleArchetypeDestroy, "System.Mass.Observer.Fragment.SingleEntitySingleArchetypeDestroy");
+IMPLEMENT_AI_INSTANT_TEST(FFragmentTest_SingleEntitySingleArchetypeDestroy, "System.Mass.Observer.Fragment.SingleEntitySingleArchetypeDestroy");
 
-struct FFragmentObserverTest_MultipleArchetypeAdd : FFragmentTestBase
+struct FFragmentTest_MultipleArchetypeAdd : FFragmentTestBase
 {
-	FFragmentObserverTest_MultipleArchetypeAdd() { OperationObserved = EMassObservedOperation::Add; }
+	FFragmentTest_MultipleArchetypeAdd() { OperationObserved = EMassObservedOperation::Add; }
 
 	virtual bool PerformOperation() override
 	{
@@ -404,11 +501,11 @@ struct FFragmentObserverTest_MultipleArchetypeAdd : FFragmentTestBase
 		return true;
 	}
 };
-IMPLEMENT_AI_INSTANT_TEST(FFragmentObserverTest_MultipleArchetypeAdd, "System.Mass.Observer.Fragment.MultipleArchetypesAdd");
+IMPLEMENT_AI_INSTANT_TEST(FFragmentTest_MultipleArchetypeAdd, "System.Mass.Observer.Fragment.MultipleArchetypesAdd");
 
-struct FFragmentObserverTest_MultipleArchetypeRemove : FFragmentTestBase
+struct FFragmentTest_MultipleArchetypeRemove : FFragmentTestBase
 {
-	FFragmentObserverTest_MultipleArchetypeRemove() { OperationObserved = EMassObservedOperation::Remove; }
+	FFragmentTest_MultipleArchetypeRemove() { OperationObserved = EMassObservedOperation::Remove; }
 
 	virtual bool PerformOperation() override
 	{
@@ -427,11 +524,11 @@ struct FFragmentObserverTest_MultipleArchetypeRemove : FFragmentTestBase
 		return true;
 	}
 };
-IMPLEMENT_AI_INSTANT_TEST(FFragmentObserverTest_MultipleArchetypeRemove, "System.Mass.Observer.Fragment.MultipleArchetypesRemove");
+IMPLEMENT_AI_INSTANT_TEST(FFragmentTest_MultipleArchetypeRemove, "System.Mass.Observer.Fragment.MultipleArchetypesRemove");
 
-struct FFragmentObserverTest_MultipleArchetypeDestroy : FFragmentTestBase
+struct FFragmentTest_MultipleArchetypeDestroy : FFragmentTestBase
 {
-	FFragmentObserverTest_MultipleArchetypeDestroy() { OperationObserved = EMassObservedOperation::Remove; }
+	FFragmentTest_MultipleArchetypeDestroy() { OperationObserved = EMassObservedOperation::Remove; }
 
 	virtual bool PerformOperation() override
 	{
@@ -450,11 +547,11 @@ struct FFragmentObserverTest_MultipleArchetypeDestroy : FFragmentTestBase
 		return true;
 	}
 };
-IMPLEMENT_AI_INSTANT_TEST(FFragmentObserverTest_MultipleArchetypeDestroy, "System.Mass.Observer.Fragment.MultipleArchetypesDestroy");
+IMPLEMENT_AI_INSTANT_TEST(FFragmentTest_MultipleArchetypeDestroy, "System.Mass.Observer.Fragment.MultipleArchetypesDestroy");
 
-struct FFragmentObserverTest_EntityCreation_Individual : FFragmentTestBase
+struct FFragmentTest_EntityCreation_Individual : FFragmentTestBase
 {
-	FFragmentObserverTest_EntityCreation_Individual() { OperationObserved = EMassObservedOperation::Add; }
+	FFragmentTest_EntityCreation_Individual() { OperationObserved = EMassObservedOperation::Add; }
 
 	virtual bool InstantTest() override
 	{
@@ -496,11 +593,11 @@ struct FFragmentObserverTest_EntityCreation_Individual : FFragmentTestBase
 		return true;
 	}
 };
-IMPLEMENT_AI_INSTANT_TEST(FFragmentObserverTest_EntityCreation_Individual, "System.Mass.Observer.Create.FragmentSingleEntity");
+IMPLEMENT_AI_INSTANT_TEST(FFragmentTest_EntityCreation_Individual, "System.Mass.Observer.Create.FragmentSingleEntity");
 
-struct FFragmentObserverTest_EntityCreation_Individuals : FFragmentTestBase
+struct FFragmentTest_EntityCreation_Individuals : FFragmentTestBase
 {
-	FFragmentObserverTest_EntityCreation_Individuals() { OperationObserved = EMassObservedOperation::Add; }
+	FFragmentTest_EntityCreation_Individuals() { OperationObserved = EMassObservedOperation::Add; }
 
 	virtual bool InstantTest() override
 	{
@@ -525,117 +622,9 @@ struct FFragmentObserverTest_EntityCreation_Individuals : FFragmentTestBase
 		return true;
 	}
 };
-IMPLEMENT_AI_INSTANT_TEST(FFragmentObserverTest_EntityCreation_Individuals, "System.Mass.Observer.Create.FragmentInvididualEntities");
-
-//-----------------------------------------------------------------------------
-// creation context 
-//-----------------------------------------------------------------------------
-struct FCreationContextTest : FEntityTestBase
-{
-	virtual bool InstantTest() override
-	{
-		constexpr int32 IntEntitiesToSpawnCount = 6;
-		constexpr int32 FloatEntitiesToSpawnCount = 7;
-
-		TArray<FMassEntityHandle> Entities;
-		TSharedRef<FMassEntityManager::FEntityCreationContext> CreationContextInt = EntityManager->BatchCreateEntities(IntsArchetype, IntEntitiesToSpawnCount, Entities);
-		TSharedRef<FMassEntityManager::FEntityCreationContext> CreationContextFloat = EntityManager->BatchCreateEntities(FloatsArchetype, FloatEntitiesToSpawnCount, Entities);
-		const int32 NumDifferentArchetypesUsed = 2;
-
-		AITEST_EQUAL(TEXT("Two back to back entity creation operations should result in the same creation context"), CreationContextInt, CreationContextFloat);
-		AITEST_FALSE(TEXT("CreationContext's entity collection should be invalidated at this moment"), CreationContextInt->DebugAreEntityCollectionsUpToDate());
-
-		TConstArrayView<FMassArchetypeEntityCollection> EntityCollections = CreationContextInt->GetEntityCollections();
-		AITEST_EQUAL(TEXT("We expect the number of resulting collections to match expectations"), EntityCollections.Num(), NumDifferentArchetypesUsed);
-
-		return true;
-	}
-};
-IMPLEMENT_AI_INSTANT_TEST(FCreationContextTest, "System.Mass.CreationContext.Append");
-
-struct FCreationContextTest_ManualCreate : FEntityTestBase
-{
-	virtual bool InstantTest() override
-	{
-		constexpr int32 IntEntitiesToSpawnCount = 6;
-		constexpr int32 FloatEntitiesToSpawnCount = 7;
-		int NumDifferentArchetypesUsed = 0;
-
-		TArray<FMassEntityHandle> Entities;
-		TSharedRef<FMassEntityManager::FEntityCreationContext> ObtainedContext = EntityManager->GetOrMakeCreationContext();
-		{
-			TSharedRef<FMassEntityManager::FEntityCreationContext> ObtainedContextCopy = EntityManager->GetOrMakeCreationContext();
-			AITEST_EQUAL(TEXT("Two back to back creation context fetching should result in the same instance"), ObtainedContext, ObtainedContextCopy);
-		}
-
-		{
-			TSharedRef<FMassEntityManager::FEntityCreationContext> CreationContextInt = EntityManager->BatchCreateEntities(IntsArchetype, IntEntitiesToSpawnCount, Entities);
-			AITEST_EQUAL(TEXT("Creating entities should return the original context"), ObtainedContext, CreationContextInt);
-			++NumDifferentArchetypesUsed;
-		}
-		
-		AITEST_TRUE(TEXT("CreationContext's entity collection should be still valid at this moment since we only added one entity collection/array")
-			, ObtainedContext->DebugAreEntityCollectionsUpToDate());
-
-		{
-			TSharedRef<FMassEntityManager::FEntityCreationContext> TempContext = EntityManager->BatchCreateEntities(IntsArchetype, IntEntitiesToSpawnCount, Entities);
-			AITEST_EQUAL(TEXT("Creating entities should return the original context"), ObtainedContext, TempContext);
-
-			AITEST_FALSE(TEXT("CreationContext's entity collection should be invalidated at this moment")
-				, TempContext->DebugAreEntityCollectionsUpToDate());
-		}
-
-		TConstArrayView<FMassArchetypeEntityCollection> EntityCollections = ObtainedContext->GetEntityCollections();
-		AITEST_EQUAL(TEXT("We expect the number of resulting collections to match expectations"), EntityCollections.Num(), NumDifferentArchetypesUsed);
-
-		return true;
-	}
-};
-IMPLEMENT_AI_INSTANT_TEST(FCreationContextTest_ManualCreate, "System.Mass.CreationContext.ManualCreate");
-
-struct FCreationContextTest_ManualBuild : FEntityTestBase
-{
-	virtual bool InstantTest() override
-	{
-		constexpr int32 FloatEntitiesToSpawnCount = 7;
-		int NumDifferentArchetypesUsed = 0;
-
-		TArray<FTestFragment_Float> Payload;
-		for (int Index = 0; Index < FloatEntitiesToSpawnCount; ++Index)
-		{ 
-			Payload.Add(FTestFragment_Float(float(Index)));
-		}
-
-		TSharedRef<FMassEntityManager::FEntityCreationContext> ObtainedContext = EntityManager->GetOrMakeCreationContext();
-		
-		TArray<FMassEntityHandle> Entities;
-		EntityManager->BatchReserveEntities(FloatEntitiesToSpawnCount, Entities);
-
-		FStructArrayView PaloadView(Payload);
-		TArray<FMassArchetypeEntityCollectionWithPayload> EntityCollections;
-		FMassArchetypeEntityCollectionWithPayload::CreateEntityRangesWithPayload(*EntityManager, Entities, FMassArchetypeEntityCollection::NoDuplicates
-			, FMassGenericPayloadView(MakeArrayView(&PaloadView, 1)), EntityCollections);
-
-		checkf(EntityCollections.Num() <= 1, TEXT("We expect TargetEntities to only contain archetype-less entities, ones that need to be \'build\'"));
-
-		{
-			TSharedRef<FMassEntityManager::FEntityCreationContext> CreationContext = EntityManager->BatchBuildEntities(EntityCollections[0], FMassFragmentBitSet(*FTestFragment_Float::StaticStruct()));
-			AITEST_EQUAL(TEXT("Creating entities should return the original context"), ObtainedContext, CreationContext);
-			++NumDifferentArchetypesUsed;
-		}
-
-		AITEST_TRUE(TEXT("CreationContext's entity collection should be still valid at this moment since we only added one entity collection/array")
-			, ObtainedContext->DebugAreEntityCollectionsUpToDate());
-
-		TConstArrayView<FMassArchetypeEntityCollection> ContextEntityCollections = ObtainedContext->GetEntityCollections();
-		AITEST_EQUAL(TEXT("We expect the number of resulting collections to match expectations"), ContextEntityCollections.Num(), NumDifferentArchetypesUsed);
-
-		return true;
-	}
-};
-IMPLEMENT_AI_INSTANT_TEST(FCreationContextTest_ManualBuild, "System.Mass.CreationContext.ManualBuild");
+IMPLEMENT_AI_INSTANT_TEST(FFragmentTest_EntityCreation_Individuals, "System.Mass.Observer.Create.FragmentInvididualEntities");
 
-} // FMassObserverTest
+} // UE::Mass::Test::Observers
 
 UE_ENABLE_OPTIMIZATION_SHIP
 
