[Mass] Added initial unit tests for newly added Mass Type Manager. Also additions to Mass commands unit test.

#rnx

diff --git a/Engine/Source/Developer/MassEntityTestSuite/Private/MassCommandsTest.cpp b/Engine/Source/Developer/MassEntityTestSuite/Private/MassCommandsTest.cpp
--- a/Engine/Source/Developer/MassEntityTestSuite/Private/MassCommandsTest.cpp
+++ b/Engine/Source/Developer/MassEntityTestSuite/Private/MassCommandsTest.cpp
@@ -81,26 +81,44 @@ struct FCommands_FragmentMemoryCleanup : FEntityTestBase
 
 		AITEST_EQUAL(TEXT("All entities created should be in ArrayArchetype"), EntityManager->DebugGetArchetypeEntitiesCount(ArrayArchetype), Entities.Num());
 
-		TArray<int32> EntitiesModified;
-		for (int i = 0; i < Count; ++i)
+		TArray<int32> EntitiesWithArray;
+		for (int EntityIndex = 0; EntityIndex < Count; ++EntityIndex)
 		{
 			if (FMath::FRand() < 0.2)
 			{	
 				FTestFragment_Array A;
-				A.Value.Add(i);
-				EntityManager->Defer().PushCommand<FMassCommandAddFragmentInstances>(Entities[i], A);
-				EntityManager->Defer().AddFragment<FTestFragment_Int>(Entities[i]);
-				EntitiesModified.Add(i);
+				A.Value.Add(EntityIndex);
+				EntityManager->Defer().PushCommand<FMassCommandAddFragmentInstances>(Entities[EntityIndex], A);
+				EntityManager->Defer().AddFragment<FTestFragment_Int>(Entities[EntityIndex]);
+				EntitiesWithArray.Add(EntityIndex);
 			}
 		}
 
 		EntityManager->FlushCommands();
 
-		for (int32 i : EntitiesModified)
+		for (int32 EntityIndex : EntitiesWithArray)
 		{
-			FMassEntityView View(ArrayIntArchetype, Entities[i]);
+			FMassEntityView View(ArrayIntArchetype, Entities[EntityIndex]);
 			AITEST_EQUAL(TEXT("Should have predicted values"), View.GetFragmentData<FTestFragment_Array>().Value.Num(), 1);
-			AITEST_EQUAL(TEXT("Should have predicted values"), View.GetFragmentData<FTestFragment_Array>().Value[0], i);
+			AITEST_EQUAL(TEXT("Should have predicted values"), View.GetFragmentData<FTestFragment_Array>().Value[0], EntityIndex);
+		}
+
+		// not move things a round by adding yet another fragment. That will force moving of some array-hosting fragments
+		for (int EntityIndex = 0; EntityIndex < Count; ++EntityIndex)
+		{
+			if (FMath::FRand() < 0.5)
+			{	
+				EntityManager->Defer().AddFragment<FTestFragment_Float>(Entities[EntityIndex]);
+			}
+		}
+
+		EntityManager->FlushCommands();
+
+		for (int32 EntityIndex : EntitiesWithArray)
+		{
+			FMassEntityView View(*EntityManager, Entities[EntityIndex]);
+			AITEST_EQUAL(TEXT("Potentially moved array fragment should have predicted values"), View.GetFragmentData<FTestFragment_Array>().Value.Num(), 1);
+			AITEST_EQUAL(TEXT("Potentially moved array fragment should have predicted values"), View.GetFragmentData<FTestFragment_Array>().Value[0], EntityIndex);
 		}
 
 		return true;
diff --git a/Engine/Source/Developer/MassEntityTestSuite/Private/MassTypeManagerTest.cpp b/Engine/Source/Developer/MassEntityTestSuite/Private/MassTypeManagerTest.cpp
new file mode 100644
--- /dev/null
+++ b/Engine/Source/Developer/MassEntityTestSuite/Private/MassTypeManagerTest.cpp
@@ -0,0 +1,32 @@
+// Copyright Epic Games, Inc. All Rights Reserved.
+
+#include "AITestsCommon.h"
+#include "MassEntityTestTypes.h"
+#include "MassTypeManager.h"
+
+
+#define LOCTEXT_NAMESPACE "MassTest"
+
+UE_DISABLE_OPTIMIZATION_SHIP
+
+namespace UE::Mass::Test
+{
+	struct FTypeManager_Subsystem : FEntityTestBase
+	{
+		virtual bool InstantTest() override
+		{
+			const FSubsystemTypeTraits TestSubsystemTraits = FSubsystemTypeTraits::Make<UMassTestWorldSubsystem>();
+
+			AITEST_EQUAL("Subsystem bGameThreadOnly value", TestSubsystemTraits.bGameThreadOnly, TMassExternalSubsystemTraits<UMassTestWorldSubsystem>::GameThreadOnly)
+			AITEST_EQUAL("Subsystem ThreadSafeWrite value", TestSubsystemTraits.bThreadSafeWrite, TMassExternalSubsystemTraits<UMassTestWorldSubsystem>::ThreadSafeWrite)
+
+			const FSharedFragmentTypeTraits TestSharedFragmentTraits = FSharedFragmentTypeTraits::Make<FTestSharedFragment_Int>();
+			AITEST_EQUAL("Shared fragment bGameThreadOnly value", TestSubsystemTraits.bGameThreadOnly, TMassExternalSubsystemTraits<UMassTestWorldSubsystem>::GameThreadOnly)
+
+			return true;
+		}
+	};
+	IMPLEMENT_AI_INSTANT_TEST(FTypeManager_Subsystem, "System.Mass.TypeManager.StaticSubsystem");
+} // UE::Mass::Test
+
+UE_ENABLE_OPTIMIZATION_SHIP
