[Mass] Issue 791295: build fix, added missing #undef LOCTEXT_NAMESPACE

#rnx

diff --git a/Engine/Source/Developer/MassEntityTestSuite/Private/MassTypeManagerTest.cpp b/Engine/Source/Developer/MassEntityTestSuite/Private/MassTypeManagerTest.cpp
--- a/Engine/Source/Developer/MassEntityTestSuite/Private/MassTypeManagerTest.cpp
+++ b/Engine/Source/Developer/MassEntityTestSuite/Private/MassTypeManagerTest.cpp
@@ -4,7 +4,6 @@
 #include "MassEntityTestTypes.h"
 #include "MassTypeManager.h"
 
-
 #define LOCTEXT_NAMESPACE "MassTest"
 
 UE_DISABLE_OPTIMIZATION_SHIP
@@ -30,3 +29,5 @@ namespace UE::Mass::Test
 } // UE::Mass::Test
 
 UE_ENABLE_OPTIMIZATION_SHIP
+
+#undef LOCTEXT_NAMESPACE
\ No newline at end of file
