[Mass] another Mass EntityIterator test added

#rnx

diff --git a/Engine/Source/Developer/MassEntityTestSuite/Private/MassEntityIteratorTest.cpp b/Engine/Source/Developer/MassEntityTestSuite/Private/MassEntityIteratorTest.cpp
--- a/Engine/Source/Developer/MassEntityTestSuite/Private/MassEntityIteratorTest.cpp
+++ b/Engine/Source/Developer/MassEntityTestSuite/Private/MassEntityIteratorTest.cpp
@@ -137,6 +137,30 @@ struct FIterator_QuerylessIterator : FEntityTestBase
 };
 IMPLEMENT_AI_INSTANT_TEST(FIterator_QuerylessIterator, "System.Mass.Entity.Iterator.Queryless");
 
+struct FIterator_ProcessorlessIterator : FEntityTestBase
+{
+	virtual bool InstantTest() override
+	{
+		FMassExecutionContext LocalContext(*EntityManager.Get());
+		{
+			AITEST_SCOPED_CHECK("no entity query is being executed", 1);
+			FMassExecutionContext::FEntityIterator FailedIterator = LocalContext.CreateEntityIterator();
+
+			AITEST_FALSE("(Not) Created iterator is valid", bool(FailedIterator));
+
+			int32 NumIterations = 0;
+			for (; FailedIterator; ++FailedIterator)
+			{
+				++NumIterations;
+			}
+			AITEST_EQUAL("Number of iterations with an invalid iterator", NumIterations, 0);
+		}
+		
+		return true;
+	}
+};
+IMPLEMENT_AI_INSTANT_TEST(FIterator_ProcessorlessIterator, "System.Mass.Entity.Iterator.Processorless");
+
 } // UE::Mass::Test
 
 #endif // WITH_MASSENTITY_DEBUG
