[Mass] EntityBuilder support for using Instanced Structs (will be used in BP interface later).

#rnx

diff --git a/Engine/Source/Developer/MassEntityTestSuite/Private/MassEntityBuilderTest.cpp b/Engine/Source/Developer/MassEntityTestSuite/Private/MassEntityBuilderTest.cpp
--- a/Engine/Source/Developer/MassEntityTestSuite/Private/MassEntityBuilderTest.cpp
+++ b/Engine/Source/Developer/MassEntityTestSuite/Private/MassEntityBuilderTest.cpp
@@ -599,6 +599,42 @@ namespace UE::Mass::Test::EntityBuilder
 		}
 	};
 	IMPLEMENT_AI_INSTANT_TEST(FAppendFromEntity, "System.Mass.EntityBuilder.Append");
+
+	struct FUsingInstancedStructs : FAllElementsUsed
+	{
+		virtual bool SetUp() override
+		{
+			// deliberately skipping FAllElementsUsed, we're doing a different setup here
+			if (FEntityTestBase::SetUp())
+			{
+				// quick builder just to create an entity with known properties
+				FEntityBuilder Builder(EntityManager.ToSharedRef());
+
+				FInstancedStruct ElementInstance;
+
+				ElementInstance.InitializeAs<FTestFragment_Int>(TestIntValue);
+				Builder.Add(ElementInstance);
+				
+				ElementInstance.InitializeAs<FTestFragment_Float>(TestFloatValue);
+				Builder.Add(MoveTemp(ElementInstance));
+
+				ElementInstance.InitializeAs<FTestSharedFragment_Float>(TestSharedFloatValue);
+				Builder.Add(MoveTemp(ElementInstance));
+
+				ElementInstance.InitializeAs<FTestConstSharedFragment_Int>(TestSharedIntValue);
+				Builder.Add(ElementInstance);
+
+				// tags cannot be added as instanced structs
+				Builder.Add<FTestTag_B>();
+
+				OriginalEntity = Builder.Commit();
+
+				return true;
+			}
+			return false;
+		}
+	};
+	IMPLEMENT_AI_INSTANT_TEST(FUsingInstancedStructs, "System.Mass.EntityBuilder.WithInstancedStructs");
 } // UE::Mass::Test::EntityBuilder
 
 UE_ENABLE_OPTIMIZATION_SHIP
diff --git a/Engine/Source/Runtime/MassEntity/Private/MassEntityBuilder.cpp b/Engine/Source/Runtime/MassEntity/Private/MassEntityBuilder.cpp
--- a/Engine/Source/Runtime/MassEntity/Private/MassEntityBuilder.cpp
+++ b/Engine/Source/Runtime/MassEntity/Private/MassEntityBuilder.cpp
@@ -371,4 +371,39 @@ void FEntityBuilder::CacheSharedFragmentValue()
 	}
 }
 
+template<typename T>
+FORCEINLINE FEntityBuilder& FEntityBuilder::AddInternal(T&& ElementInstance)
+{
+	if (const UScriptStruct* Type = ElementInstance.GetScriptStruct())
+	{
+		if (Type->IsChildOf(FMassFragment::StaticStruct()))
+		{
+			Composition.Fragments.Add(*Type);
+			Fragments.Add(Forward<T>(ElementInstance));
+		}
+		else if (Type->IsChildOf(FMassSharedFragment::StaticStruct()))
+		{
+			Composition.SharedFragments.Add(*Type);
+			SharedFragments.Add(Forward<T>(ElementInstance));
+		}
+		else if (Type->IsChildOf(FMassConstSharedFragment::StaticStruct()))
+		{
+			Composition.ConstSharedFragments.Add(*Type);
+			ConstSharedFragments.Add(Forward<T>(ElementInstance));
+		}
+	}
+
+	return *this;
+}
+
+FEntityBuilder& FEntityBuilder::Add(const FInstancedStruct& ElementInstance)
+{
+	return AddInternal(ElementInstance);
+}
+
+FEntityBuilder& FEntityBuilder::Add(FInstancedStruct&& ElementInstance)
+{
+	return AddInternal(ElementInstance);
+}
+
 } // namespace UE::Mass
\ No newline at end of file
diff --git a/Engine/Source/Runtime/MassEntity/Public/MassEntityBuilder.h b/Engine/Source/Runtime/MassEntity/Public/MassEntityBuilder.h
--- a/Engine/Source/Runtime/MassEntity/Public/MassEntityBuilder.h
+++ b/Engine/Source/Runtime/MassEntity/Public/MassEntityBuilder.h
@@ -156,6 +156,12 @@ struct MASSENTITY_API FEntityBuilder
 	template<typename T, typename... TArgs>
 	FEntityBuilder& Add(TArgs&&... InArgs) requires (CElement<T> && !(CTag<T> || CChunkFragment<T>));
 
+	/**
+	 * Adds a fragment instance to the Entity Builder, treating the contents according to its type
+	 */
+	FEntityBuilder& Add(const FInstancedStruct& ElementInstance);
+	FEntityBuilder& Add(FInstancedStruct&& ElementInstance);
+
 	/**
 	 * Finds and retrieves a pointer to a fragment of type T if it exists.
 	 * @param T - The type of fragment to find.
@@ -178,6 +184,9 @@ struct MASSENTITY_API FEntityBuilder
 	/** @return whether the builder has already committed the data */
 	bool IsCommitted() const;
 
+	/** @return the EntityManager instance this entity builder is working for */
+	TSharedRef<FMassEntityManager> GetEntityManager();
+
 protected:
 	friend Private::FEntityBuilderHelper;
 
@@ -186,6 +195,9 @@ protected:
 	void InvalidateCachedData();
 
 private:
+	template<typename T>
+	FEntityBuilder& AddInternal(T&& ElementInstance);
+
 	TSharedRef<FMassEntityManager> EntityManager;
 	mutable FMassEntityHandle EntityHandle;
 
@@ -371,4 +383,9 @@ inline bool FEntityBuilder::IsCommitted() const
 	return State == EState::Committed;
 }
 
+inline TSharedRef<FMassEntityManager> FEntityBuilder::GetEntityManager()
+{
+	return EntityManager;
+}
+
 } // namespace UE::Mass
