[Mass] Added new tests for MassObservers. One exposed a flaw in current implementation - this one's disabled for now.

diff --git a/Engine/Source/Developer/MassEntityTestSuite/Private/MassObserverTest.cpp b/Engine/Source/Developer/MassEntityTestSuite/Private/MassObserverTest.cpp
--- a/Engine/Source/Developer/MassEntityTestSuite/Private/MassObserverTest.cpp
+++ b/Engine/Source/Developer/MassEntityTestSuite/Private/MassObserverTest.cpp
@@ -381,6 +381,7 @@ struct FFragmentTestBase : FEntityTestBase
 	TArray<FMassEntityHandle> AffectedEntities;
 	UMassTestProcessorBase* ObserverProcessor = nullptr;
 	EMassObservedOperation OperationObserved = EMassObservedOperation::MAX;
+	TArray<FMassEntityHandle> EntitiesFloats;
 	TArray<FMassEntityHandle> EntitiesInt;
 	TArray<FMassEntityHandle> EntitiesIntsFloat;
 	TArray<FMassEntityHandle> ExpectedEntities;
@@ -622,7 +623,77 @@ struct FFragmentTest_EntityCreation_Individuals : FFragmentTestBase
 		return true;
 	}
 };
-IMPLEMENT_AI_INSTANT_TEST(FFragmentTest_EntityCreation_Individuals, "System.Mass.Observer.Create.FragmentInvididualEntities");
+IMPLEMENT_AI_INSTANT_TEST(FFragmentTest_EntityCreation_Individuals, "System.Mass.Observer.Create.FragmentIndividualEntities");
+
+/**
+ * This test aims to verify expected behavior of observers when there's a creation context active, when composition-mutating
+ * operations are affecting entities other than the ones being created.
+ */
+struct FModificationsWhileCreationContextActive : FTagBaseOperation
+{
+	FModificationsWhileCreationContextActive() { OperationObserved = EMassObservedOperation::Add; }
+
+	virtual bool InstantTest() override
+	{
+		constexpr int32 EntitiesToSpawnInFirstBatch = 3;
+		constexpr int32 EntitiesToSpawnInSecondBatch = 5;
+
+		EntityManager->BatchCreateEntities(IntsArchetype, EntitiesToSpawnInFirstBatch, EntitiesInt);
+		FMassArchetypeEntityCollection InitialEntitiesCollection(IntsArchetype, EntitiesInt, FMassArchetypeEntityCollection::EDuplicatesHandling::NoDuplicates);
+
+		FMassObserverManager& ObserverManager = EntityManager->GetObserverManager();
+		ObserverManager.AddObserverInstance(*FTagStruct::StaticStruct(), OperationObserved, *ObserverProcessor);
+
+		{
+			TSharedRef<FMassEntityManager::FEntityCreationContext> CreationContext = EntityManager->BatchCreateEntities(IntsArchetype, EntitiesToSpawnInSecondBatch, EntitiesInt);
+			ensure(EntitiesInt.Num() == EntitiesToSpawnInFirstBatch + EntitiesToSpawnInSecondBatch);
+			// note that the observers' behavior regarding the entities just created gets tested by FEntityCreation_Batched test above
+			// we're testing only the behavior related to the previously created entities here
+
+			EntityManager->BatchChangeTagsForEntities(MakeArrayView(&InitialEntitiesCollection, 1)
+				, FMassTagBitSet(*FTagStruct::StaticStruct()), FMassTagBitSet());
+
+			AITEST_TRUE(TEXT("The tag observer is not expected to run yet"), AffectedEntities.Num() == 0);
+			
+			// should it, though?
+			//AITEST_FALSE(TEXT("CreationContext's entity collection should be invalidated at this moment"), CreationContext->DebugAreEntityCollectionsUpToDate());
+		}
+		AITEST_TRUE(TEXT("The tag observer is expected to run just after FEntityCreationContext's destruction"), AffectedEntities.Num() > 0);
+		AITEST_EQUAL(TEXT("The tag observer is expected to process only the original entities, that had a tag added to them"), AffectedEntities.Num(), EntitiesToSpawnInFirstBatch);
+
+		return true;
+	}
+};
+IMPLEMENT_AI_INSTANT_TEST(FModificationsWhileCreationContextActive, "System.Mass.Observer.Create.ModificationsToOtherEntities");
+
+// The scenario being tested:
+// 1. Create entities with Float fragment
+// 2. Add an unobserved tag A - results in created entities changing archetype
+// 3. Release the creation context - we expect Float observers to trigger
+struct FCreatedEntitiesUnobservedCompositionChange : FFragmentTestBase
+{
+	FCreatedEntitiesUnobservedCompositionChange() { OperationObserved = EMassObservedOperation::Add; }
+
+	virtual bool InstantTest() override
+	{
+		constexpr int32 EntitiesToSpawnCount = 6;
+
+		FMassObserverManager& ObserverManager = EntityManager->GetObserverManager();
+		ObserverManager.AddObserverInstance(*FFragmentStruct::StaticStruct(), OperationObserved, *ObserverProcessor);
+		{
+			TSharedRef<FMassEntityManager::FEntityCreationContext> CreationContext = EntityManager->BatchCreateEntities(FloatsArchetype, EntitiesToSpawnCount, EntitiesFloats);
+
+			// add unobserved tag
+			EntityManager->BatchChangeTagsForEntities(CreationContext->GetEntityCollections(), FMassTagBitSet(*FTestTag_A::StaticStruct()), FMassTagBitSet());
+
+			AITEST_EQUAL(TEXT("The fragment observer is not expected to run yet"), AffectedEntities.Num(), 0);
+		}
+		AITEST_EQUAL(TEXT("The fragment observer is expected to run just after FEntityCreationContext's destruction"), AffectedEntities.Num(), EntitiesToSpawnCount);
+
+		return true;
+	}
+};
+IMPLEMENT_AI_INSTANT_TEST(FCreatedEntitiesUnobservedCompositionChange, "System.Mass.Observer.Create.UnobservedCompositionChange");
 
 } // UE::Mass::Test::Observers
 
diff --git a/Engine/Source/Runtime/MassEntity/Private/MassObserverManager.cpp b/Engine/Source/Runtime/MassEntity/Private/MassObserverManager.cpp
--- a/Engine/Source/Runtime/MassEntity/Private/MassObserverManager.cpp
+++ b/Engine/Source/Runtime/MassEntity/Private/MassObserverManager.cpp
@@ -162,6 +162,11 @@ bool FMassObserverManager::OnPostEntitiesCreated(FMassProcessingContext& Process
 
 bool FMassObserverManager::OnPostEntitiesCreated(FMassProcessingContext& ProcessingContext, TConstArrayView<FMassArchetypeEntityCollection> EntityCollections)
 {
+	if (!ensureMsgf(EntityManager.IsDuringEntityCreation() == false, TEXT("Entity creation not done yet.")))
+	{
+		return false;
+	}
+
 	TRACE_CPUPROFILER_EVENT_SCOPE_STR("OnPostEntitiesCreated")
 
 	check(ProcessingContext.EntityManager);
@@ -211,6 +216,8 @@ bool FMassObserverManager::OnCompositionChanged(FMassProcessingContext& Processi
 {
 	using UE::Mass::ObserverManager::Tweakables::InlineAllocatorElementsForOverlapTypes;
 
+	ensureMsgf(EntityCollection.IsUpToDate(), TEXT("Out-of-date FMassArchetypeEntityCollection used. Stored information is unreliable."));
+
 	const FMassFragmentBitSet FragmentOverlap = ObservedFragments[(uint8)Operation].GetOverlap(CompositionDelta.Fragments);
 	const bool bHasFragmentsOverlap = !FragmentOverlap.IsEmpty();
 	const FMassTagBitSet TagOverlap = ObservedTags[(uint8)Operation].GetOverlap(CompositionDelta.Tags);
