[Mass] Added a way to export actual entity handles from FMassArchetypeEntityCollection

diff --git a/Engine/Source/Developer/MassEntityTestSuite/Private/MassChunkCollectionTest.cpp b/Engine/Source/Developer/MassEntityTestSuite/Private/MassChunkCollectionTest.cpp
--- a/Engine/Source/Developer/MassEntityTestSuite/Private/MassChunkCollectionTest.cpp
+++ b/Engine/Source/Developer/MassEntityTestSuite/Private/MassChunkCollectionTest.cpp
@@ -6,6 +6,8 @@
 #include "MassExecutionContext.h"
 #include "Algo/RandomShuffle.h"
 #include "MassExecutor.h"
+#include "Algo/Compare.h"
+#include "Algo/RemoveIf.h"
 
 #define LOCTEXT_NAMESPACE "MassTest"
 
@@ -515,6 +517,38 @@ struct FEntityCollection_ZeroLength_BatchSetFragmentValues : FEntityTestBase
 };
 IMPLEMENT_AI_INSTANT_TEST(FEntityCollection_ZeroLength_BatchSetFragmentValues, "System.Mass.Entity.ZeroLengthRanges.BatchSetFragmentValues");
 
+struct FEntityCollection_ExportHandles_Created : FEntityTestBase
+{
+	virtual bool InstantTest() override
+	{
+		CA_ASSUME(EntityManager);
+		const FMassArchetypeHandle SourceArchetype = FloatsArchetype;
+		const int32 EntitiesPerChunk = FMassArchetypeHelper::ArchetypeDataFromHandleChecked(SourceArchetype).GetNumEntitiesPerChunk();
+		const int32 NumChunks = 3;
+		const int32 EntityCount = NumChunks * EntitiesPerChunk;
+		TArray<FMassEntityHandle> Entities;
+		EntityManager->BatchCreateEntities(SourceArchetype, {}, EntityCount, Entities);
+
+		// deterministically removing half of the entities
+		FRandomStream RandomStream(0);
+		const int32 RemainingCount = Algo::StableRemoveIf(Entities, [&RandomStream](const FMassEntityHandle& Element)
+		{
+			return RandomStream.FRand() < 0.5f;
+		});
+		Entities.SetNum(RemainingCount);
+
+		FMassArchetypeEntityCollection Collection(SourceArchetype, Entities, FMassArchetypeEntityCollection::EDuplicatesHandling::NoDuplicates);
+
+		TArray<FMassEntityHandle> ExportedHandles;
+		Collection.ExportEntityHandles(ExportedHandles);
+
+		AITEST_TRUE("Exported handles are the same as the input data", Algo::Compare(ExportedHandles, Entities));
+
+		return true;
+	}
+};
+IMPLEMENT_AI_INSTANT_TEST(FEntityCollection_ExportHandles_Created, "System.Mass.EntityCollection.ExportHandles.CreatedEntities");
+
 #endif // WITH_MASSENTITY_DEBUG
 //
 }
diff --git a/Engine/Source/Runtime/MassEntity/Internal/MassArchetypeData.h b/Engine/Source/Runtime/MassEntity/Internal/MassArchetypeData.h
--- a/Engine/Source/Runtime/MassEntity/Internal/MassArchetypeData.h
+++ b/Engine/Source/Runtime/MassEntity/Internal/MassArchetypeData.h
@@ -57,7 +57,14 @@ public:
 	// Returns the Entity array element at the specified index
 	FMassEntityHandle& GetEntityArrayElementRef(int32 ChunkBase, int32 IndexWithinChunk)
 	{
-		return ((FMassEntityHandle*)(RawMemory + ChunkBase))[IndexWithinChunk];
+		checkSlow(reinterpret_cast<SIZE_T>(RawMemory + ChunkBase + IndexWithinChunk * sizeof(FMassEntityHandle))< AllocSize);
+		return reinterpret_cast<FMassEntityHandle*>(RawMemory + ChunkBase)[IndexWithinChunk];
+	}
+
+	const FMassEntityHandle* GetEntityArray(int32 ChunkBase) const
+	{
+		checkSlow(reinterpret_cast<SIZE_T>(RawMemory + ChunkBase)< AllocSize);
+		return reinterpret_cast<const FMassEntityHandle*>(RawMemory + ChunkBase);
 	}
 
 	uint8* GetRawMemory() const
@@ -326,6 +333,8 @@ public:
 
 	SIZE_T GetAllocatedSize() const;
 
+	void ExportEntityHandles(const TConstArrayView<FMassArchetypeEntityCollection::FArchetypeEntityRange> Ranges, TArray<FMassEntityHandle>& InOutHandles) const;
+
 	// Converts the list of fragments into a user-readable debug string
 	FString DebugGetDescription() const;
 
diff --git a/Engine/Source/Runtime/MassEntity/Private/MassArchetypeData.cpp b/Engine/Source/Runtime/MassEntity/Private/MassArchetypeData.cpp
--- a/Engine/Source/Runtime/MassEntity/Private/MassArchetypeData.cpp
+++ b/Engine/Source/Runtime/MassEntity/Private/MassArchetypeData.cpp
@@ -909,6 +909,27 @@ SIZE_T FMassArchetypeData::GetAllocatedSize() const
 		FragmentIndexMap.GetAllocatedSize();
 }
 
+void FMassArchetypeData::ExportEntityHandles(const TConstArrayView<FMassArchetypeEntityCollection::FArchetypeEntityRange> Ranges, TArray<FMassEntityHandle>& InOutHandles) const
+{
+	int32 TotalEntities = 0;
+	for (const FMassArchetypeEntityCollection::FArchetypeEntityRange& Range : Ranges)
+	{
+		TotalEntities += Range.Length;
+	}
+	int32 StartIndex = InOutHandles.Num();
+	InOutHandles.AddUninitialized(TotalEntities);
+
+	for (const FMassArchetypeEntityCollection::FArchetypeEntityRange& Range : Ranges)
+	{
+		check(Chunks.IsValidIndex(Range.ChunkIndex))
+		const FMassEntityHandle* EntitiesArray = Chunks[Range.ChunkIndex].GetEntityArray(EntityListOffsetWithinChunk);
+
+		FMemory::Memcpy(&InOutHandles[StartIndex], &EntitiesArray[Range.SubchunkStart], Range.Length * sizeof(FMassEntityHandle));
+
+		StartIndex += Range.Length;
+	}
+}
+
 FString FMassArchetypeData::DebugGetDescription() const
 {
 #if WITH_MASSENTITY_DEBUG
diff --git a/Engine/Source/Runtime/MassEntity/Private/MassArchetypeTypes.cpp b/Engine/Source/Runtime/MassEntity/Private/MassArchetypeTypes.cpp
--- a/Engine/Source/Runtime/MassEntity/Private/MassArchetypeTypes.cpp
+++ b/Engine/Source/Runtime/MassEntity/Private/MassArchetypeTypes.cpp
@@ -207,6 +207,16 @@ bool FMassArchetypeEntityCollection::IsSame(const FMassArchetypeEntityCollection
 	return true;
 }
 
+void FMassArchetypeEntityCollection::ExportEntityHandles(TArray<FMassEntityHandle>& InOutHandles) const
+{
+	const FMassArchetypeData* ArchetypeData = FMassArchetypeHelper::ArchetypeDataFromHandle(Archetype);
+	if (ensureMsgf(ArchetypeData, TEXT("This operation is supported only for collections containing created entities (i.e. not 'reserved'")))
+	{
+		checkf(IsUpToDate(), TEXT("The entity collection is out of that. The following action might crash and if not, the results might be incorrect."));
+		ArchetypeData->ExportEntityHandles(Ranges, InOutHandles);
+	}
+}
+
 //-----------------------------------------------------------------------------
 // FMassArchetypeEntityCollectionWithPayload
 //-----------------------------------------------------------------------------
diff --git a/Engine/Source/Runtime/MassEntity/Public/MassArchetypeTypes.h b/Engine/Source/Runtime/MassEntity/Public/MassArchetypeTypes.h
--- a/Engine/Source/Runtime/MassEntity/Public/MassArchetypeTypes.h
+++ b/Engine/Source/Runtime/MassEntity/Public/MassArchetypeTypes.h
@@ -178,6 +178,12 @@ public:
 	/** The comparison function that checks if Other is identical to this. Intended for diagnostics/debugging. */
 	bool IsSame(const FMassArchetypeEntityCollection& Other) const;
 
+	/**
+	 * Converts stored entity ranges to FMassEntityHandles and appends them to InOutHandles.
+	 * Note that the operation is only supported for already created entities (i.e. not "reserved")
+	 */
+	void ExportEntityHandles(TArray<FMassEntityHandle>& InOutHandles) const;
+
 protected:
 	friend struct FMassArchetypeEntityCollectionWithPayload;
 	void BuildEntityRanges(TStridedView<const int32> TrueIndices);
