[MassTesting] updated System.Mass.Dependencies.Circular test since it went out of sync with changes to dependency solver logging.

#rnx

diff --git a/Engine/Source/Developer/MassEntityTestSuite/Private/MassDependencyTest.cpp b/Engine/Source/Developer/MassEntityTestSuite/Private/MassDependencyTest.cpp
--- a/Engine/Source/Developer/MassEntityTestSuite/Private/MassDependencyTest.cpp
+++ b/Engine/Source/Developer/MassEntityTestSuite/Private/MassDependencyTest.cpp
@@ -444,6 +444,7 @@ struct FCircularDependency : FDependencySolverBase
 	{
 		{
 			AITEST_SCOPED_CHECK(TEXT("Detected processing dependency cycle"), 1);
+			AITEST_SCOPED_CHECK(TEXT("Encountered processing dependency cycle"), 1);
 			Solve();
 		}
 		// every subsequent processor is expected to depend only on the previous one since all the processors use exactly the same resources
