[Mass] Unit tests for newly added "handle exporting" functionality

#rnx

diff --git a/Engine/Source/Developer/MassEntityTestSuite/Private/MassCommandsTest.cpp b/Engine/Source/Developer/MassEntityTestSuite/Private/MassCommandsTest.cpp
--- a/Engine/Source/Developer/MassEntityTestSuite/Private/MassCommandsTest.cpp
+++ b/Engine/Source/Developer/MassEntityTestSuite/Private/MassCommandsTest.cpp
@@ -313,6 +313,38 @@ struct FCommands_PushWhileFlushing : FEntityTestBase
 };
 IMPLEMENT_AI_INSTANT_TEST(FCommands_PushWhileFlushing, "System.Mass.Commands.PushWhileFlushing");
 
+struct FCommands_MoveHandleArrays : FEntityTestBase
+{
+	virtual bool InstantTest() override
+	{
+		const int32 EntitiesPerChunk = EntityManager->DebugGetArchetypeEntitiesCountPerChunk(FloatsIntsArchetype);
+		const int32 Count = static_cast<int32>(static_cast<float>(EntitiesPerChunk) * 2.5f);
+
+		TArray<FMassEntityHandle> Entities;
+		EntityManager->BatchCreateEntities(IntsArchetype, Count, Entities);
+		EntityManager->BatchCreateEntities(FloatsArchetype, Count, Entities);
+		
+		TArray<FMassArchetypeEntityCollection> EntityCollections;
+		UE::Mass::Utils::CreateEntityCollections(*EntityManager.Get(), Entities, FMassArchetypeEntityCollection::NoDuplicates, EntityCollections);
+		EntityManager->BatchChangeTagsForEntities(EntityCollections, FMassTagBitSet(*FTestTag_A::StaticStruct()), FMassTagBitSet());
+
+		// verify that original archetypes no longer host any entities
+		AITEST_TRUE("Original archetypes are empty after adding a tag to all entities"
+			, EntityManager->DebugGetArchetypeEntitiesCount(IntsArchetype) == 0
+				&& EntityManager->DebugGetArchetypeEntitiesCount(FloatsArchetype) == 0);
+
+		EntityManager->Defer().PushCommand<FMassCommandRemoveTag<FTestTag_A>>(MoveTemp(Entities));
+		EntityManager->FlushCommands();
+
+		AITEST_TRUE("All the entities moved back to the original archetypes"
+			, EntityManager->DebugGetArchetypeEntitiesCount(IntsArchetype) == Count
+				&& EntityManager->DebugGetArchetypeEntitiesCount(FloatsArchetype) == Count);
+		
+		return true;
+	}
+};
+IMPLEMENT_AI_INSTANT_TEST(FCommands_MoveHandleArrays, "System.Mass.Commands.MoveHandleArrays");
+
 #endif // WITH_MASSENTITY_DEBUG
 } // FMassCommandsTest
 
diff --git a/Engine/Source/Developer/MassEntityTestSuite/Private/MassQueryTest.cpp b/Engine/Source/Developer/MassEntityTestSuite/Private/MassQueryTest.cpp
--- a/Engine/Source/Developer/MassEntityTestSuite/Private/MassQueryTest.cpp
+++ b/Engine/Source/Developer/MassEntityTestSuite/Private/MassQueryTest.cpp
@@ -8,6 +8,7 @@
 #include "MassExecutor.h"
 #include "MassExecutionContext.h"
 #include "MassTypeManager.h"
+#include "Algo/Compare.h"
 
 #define LOCTEXT_NAMESPACE "MassTest"
 
@@ -908,6 +909,47 @@ struct FQueryTest_GameThreadOnly : FEntityTestBase
 };
 IMPLEMENT_AI_INSTANT_TEST(FQueryTest_GameThreadOnly, "System.Mass.Query.GameThreadOnly");
 
+struct FQueryTest_ExportHandles : FEntityTestBase
+{
+	virtual bool InstantTest() override
+	{
+		const int32 EntitiesPerChunk = EntityManager->DebugGetArchetypeEntitiesCountPerChunk(FloatsIntsArchetype);
+		const int32 Count = static_cast<int32>(static_cast<float>(EntitiesPerChunk) * 2.5f);
+
+		TArray<FMassEntityHandle> Entities;
+		EntityManager->BatchCreateEntities(IntsArchetype, Count, Entities);
+		EntityManager->BatchCreateEntities(FloatsArchetype, Count, Entities);
+		check(Entities.Num() == 2*Count);
+
+		TArray<FMassArchetypeEntityCollection> EntityCollections;
+		UE::Mass::Utils::CreateEntityCollections(*EntityManager.Get(), Entities, FMassArchetypeEntityCollection::NoDuplicates, EntityCollections);
+		EntityManager->BatchChangeTagsForEntities(EntityCollections, FMassTagBitSet(*FTestTag_A::StaticStruct()), FMassTagBitSet());
+
+		FMassEntityQuery Query(EntityManager.ToSharedRef());
+		Query.AddTagRequirement<FTestTag_A>(EMassFragmentPresence::All);
+
+		TArray<FMassEntityHandle> QueryMatchingEntities = Query.GetMatchingEntityHandles();
+
+		Entities.Sort();
+		QueryMatchingEntities.Sort();
+		AITEST_TRUE("Exported handle list contain all the expected handles", Algo::Compare(Entities, QueryMatchingEntities));
+		
+		TArray<FMassArchetypeEntityCollection> MatchingCollections = Query.CreateMatchingEntitiesCollection();
+		AITEST_EQUAL("Expected number of archetypes in resulting collections", MatchingCollections.Num(), 2);
+
+		TArray<FMassEntityHandle> HandlesFromCollections;
+		for (const FMassArchetypeEntityCollection& Collection : MatchingCollections)
+		{
+			Collection.ExportEntityHandles(HandlesFromCollections);
+		}
+		HandlesFromCollections.Sort();
+		AITEST_TRUE("Handles exported from the collections contain all the expected handles", Algo::Compare(Entities, HandlesFromCollections));
+
+		return true;
+	}
+};
+IMPLEMENT_AI_INSTANT_TEST(FQueryTest_ExportHandles, "System.Mass.Query.ExportHandles");
+
 } // FMassQueryTest
 
 #undef LOCTEXT_NAMESPACE
