CIS fix: removed the use of debug-time-only function in the new Mass unit test

#rnx

diff --git a/Engine/Source/Developer/MassEntityTestSuite/Private/MassQueryTest.cpp b/Engine/Source/Developer/MassEntityTestSuite/Private/MassQueryTest.cpp
--- a/Engine/Source/Developer/MassEntityTestSuite/Private/MassQueryTest.cpp
+++ b/Engine/Source/Developer/MassEntityTestSuite/Private/MassQueryTest.cpp
@@ -913,8 +913,8 @@ struct FQueryTest_ExportHandles : FEntityTestBase
 {
 	virtual bool InstantTest() override
 	{
-		const int32 EntitiesPerChunk = EntityManager->DebugGetArchetypeEntitiesCountPerChunk(FloatsIntsArchetype);
-		const int32 Count = static_cast<int32>(static_cast<float>(EntitiesPerChunk) * 2.5f);
+		constexpr int32 EntitiesPerChunk = 16384;
+		constexpr int32 Count = static_cast<int32>(static_cast<float>(EntitiesPerChunk) * 2.5f);
 
 		TArray<FMassEntityHandle> Entities;
 		EntityManager->BatchCreateEntities(IntsArchetype, Count, Entities);
