[Mass] Fixed FMassEntityManager::MoveEntityToAnotherArchetype not triggering observers

diff --git a/Engine/Source/Developer/MassEntityTestSuite/Private/MassObserverTest.cpp b/Engine/Source/Developer/MassEntityTestSuite/Private/MassObserverTest.cpp
--- a/Engine/Source/Developer/MassEntityTestSuite/Private/MassObserverTest.cpp
+++ b/Engine/Source/Developer/MassEntityTestSuite/Private/MassObserverTest.cpp
@@ -695,6 +695,85 @@ struct FCreatedEntitiesUnobservedCompositionChange : FFragmentTestBase
 };
 IMPLEMENT_AI_INSTANT_TEST(FCreatedEntitiesUnobservedCompositionChange, "System.Mass.Observer.Create.UnobservedCompositionChange");
 
+
+struct FMoveToAnotherArchetype_SingleEntity : FEntityTestBase
+{
+	bool bTagAdded = false;
+	bool bTagRemoved = false;
+	bool bFloatAdded = false;
+	bool bFloatRemoved = false;
+
+	virtual bool SetUp() override
+	{
+		if (FEntityTestBase::SetUp() == false)
+		{
+			return false;
+		}
+
+		FMassObserverManager& ObserverManager = EntityManager->GetObserverManager();
+
+		auto CreateObserver = [this](const TFunction<void(FMassExecutionContext& Context)>& StoreResultFunction)
+		{
+			UMassTestProcessorBase* ObserverProcessor = NewTestProcessor<UMassTestProcessorBase>(EntityManager);
+			ObserverProcessor->EntityQuery.AddRequirement<FTestFragment_Int>(EMassFragmentAccess::ReadOnly);
+			ObserverProcessor->ForEachEntityChunkExecutionFunction = StoreResultFunction;
+			return ObserverProcessor;
+		};
+
+		ObserverManager.AddObserverInstance(*FTestTag_A::StaticStruct(), EMassObservedOperation::Add
+			, *CreateObserver([this](FMassExecutionContext& Context)
+				{
+					bTagAdded = true;
+				}
+			));
+
+		ObserverManager.AddObserverInstance(*FTestTag_A::StaticStruct(), EMassObservedOperation::Remove
+			, *CreateObserver([this](FMassExecutionContext& Context)
+				{
+					bTagRemoved = true;
+				}
+			));
+
+		ObserverManager.AddObserverInstance(*FTestFragment_Float::StaticStruct(), EMassObservedOperation::Add
+			, *CreateObserver([this](FMassExecutionContext& Context)
+				{
+					bFloatAdded = true;
+				}
+			));
+
+		ObserverManager.AddObserverInstance(*FTestFragment_Float::StaticStruct(), EMassObservedOperation::Remove
+			, *CreateObserver([this](FMassExecutionContext& Context)
+				{
+					bFloatRemoved = true;
+				}
+			));
+
+		return true;
+	}
+
+	virtual bool InstantTest() override
+	{
+		const FMassEntityHandle EntityHandle = EntityManager->CreateEntity(IntsArchetype);
+
+		// create target archetype
+		const FMassArchetypeHandle TargetArchetypeHandle = EntityManager->CreateArchetype(IntsArchetype, { FTestTag_A::StaticStruct(), FTestFragment_Float::StaticStruct() });
+
+		EntityManager->MoveEntityToAnotherArchetype(EntityHandle, TargetArchetypeHandle);
+		AITEST_TRUE("Tag addition observer has been executed", bTagAdded);
+		AITEST_TRUE("Fragment addition observer has been executed", bFloatAdded);
+		AITEST_FALSE("(NOT) Tag removal observer has been executed", bTagRemoved);
+		AITEST_FALSE("(NOT) Fragment addition observer has been executed", bFloatRemoved);
+
+		// moving back to the original archetype will remove the two added elements, and should trigger observers
+		EntityManager->MoveEntityToAnotherArchetype(EntityHandle, IntsArchetype);
+		AITEST_TRUE("Tag removal observer has been executed", bTagRemoved);
+		AITEST_TRUE("Fragment addition observer has been executed", bFloatRemoved);
+
+		return true;
+	}
+};
+IMPLEMENT_AI_INSTANT_TEST(FMoveToAnotherArchetype_SingleEntity, "System.Mass.Observer.Move.UnobservedCompositionChange");
+
 } // UE::Mass::Test::Observers
 
 UE_ENABLE_OPTIMIZATION_SHIP
diff --git a/Engine/Source/Runtime/MassEntity/Private/MassEntityManager.cpp b/Engine/Source/Runtime/MassEntity/Private/MassEntityManager.cpp
--- a/Engine/Source/Runtime/MassEntity/Private/MassEntityManager.cpp
+++ b/Engine/Source/Runtime/MassEntity/Private/MassEntityManager.cpp
@@ -1885,8 +1885,27 @@ void FMassEntityManager::MoveEntityToAnotherArchetype(FMassEntityHandle Entity,
 	// Move the entity over
 	FMassArchetypeData* CurrentArchetype = GetEntityStorageInterface().GetArchetype(Entity.Index);
 	check(CurrentArchetype);
+
+	// we'll use this local var to cache whether we should execute OnCompositionChanged(EMassObservedOperation::Add)
+	// operation. We cannot execute it in the following block, because we need to update the entity storage before
+	// triggering the observers.
+	bool bExecutePostMoveObserver = false;
+	
+	if (IsAllowedToTriggerObservers())
+	{
+		const FMassArchetypeCompositionDescriptor CompositionRemoved = CurrentArchetype->GetCompositionDescriptor().CalculateDifference(NewArchetype.GetCompositionDescriptor());
+		ObserverManager.OnCompositionChanged(Entity, CompositionRemoved, EMassObservedOperation::Remove);
+		bExecutePostMoveObserver = true;
+	}
+	
 	CurrentArchetype->MoveEntityToAnotherArchetype(Entity, NewArchetype, SharedFragmentValuesOverride);
 	GetEntityStorageInterface().SetArchetypeFromShared(Entity.Index, NewArchetypeHandle.DataPtr);
+
+	if (bExecutePostMoveObserver)
+	{
+		const FMassArchetypeCompositionDescriptor CompositionAdded = NewArchetype.GetCompositionDescriptor().CalculateDifference(CurrentArchetype->GetCompositionDescriptor());
+		ObserverManager.OnCompositionChanged(Entity, CompositionAdded, EMassObservedOperation::Add);
+	}
 }
 
 void FMassEntityManager::SetEntityFragmentValues(FMassEntityHandle Entity, TArrayView<const FInstancedStruct> FragmentInstanceList)
diff --git a/Engine/Source/Runtime/MassEntity/Public/MassEntityTypes.h b/Engine/Source/Runtime/MassEntity/Public/MassEntityTypes.h
--- a/Engine/Source/Runtime/MassEntity/Public/MassEntityTypes.h
+++ b/Engine/Source/Runtime/MassEntity/Public/MassEntityTypes.h
@@ -145,6 +145,23 @@ struct FMassArchetypeCompositionDescriptor
 		ConstSharedFragments -= OtherDescriptor.ConstSharedFragments;
 	}
 
+	/**
+	 * Finds all the elements contained in `this` while missing in `OtherDescriptor` and returns
+	 * the data as a FMassArchetypeCompositionDescriptor instance
+	 */
+	FMassArchetypeCompositionDescriptor CalculateDifference(const FMassArchetypeCompositionDescriptor& OtherDescriptor) const
+	{
+		FMassArchetypeCompositionDescriptor Diff;
+
+		Diff.Fragments = Fragments - OtherDescriptor.Fragments;
+		Diff.Tags = Tags - OtherDescriptor.Tags;
+		Diff.ChunkFragments = ChunkFragments - OtherDescriptor.ChunkFragments;
+		Diff.SharedFragments = SharedFragments - OtherDescriptor.SharedFragments;
+		Diff.ConstSharedFragments = ConstSharedFragments - OtherDescriptor.ConstSharedFragments;
+
+		return Diff;
+	}
+
 	MASSENTITY_API static uint32 CalculateHash(const FMassFragmentBitSet& InFragments, const FMassTagBitSet& InTags
 		, const FMassChunkFragmentBitSet& InChunkFragments, const FMassSharedFragmentBitSet& InSharedFragmentBitSet
 		, const FMassConstSharedFragmentBitSet& InConstSharedFragmentBitSet);
