[Mass] Added tests for verifying Observers functionality when adding fragments as a part of observer's execution.

diff --git a/Engine/Source/Developer/MassEntityTestSuite/Private/MassObserverTest.cpp b/Engine/Source/Developer/MassEntityTestSuite/Private/MassObserverTest.cpp
--- a/Engine/Source/Developer/MassEntityTestSuite/Private/MassObserverTest.cpp
+++ b/Engine/Source/Developer/MassEntityTestSuite/Private/MassObserverTest.cpp
@@ -625,6 +625,83 @@ struct FFragmentTest_EntityCreation_Individuals : FFragmentTestBase
 };
 IMPLEMENT_AI_INSTANT_TEST(FFragmentTest_EntityCreation_Individuals, "System.Mass.Observer.Create.FragmentIndividualEntities");
 
+struct FObserverChangingComposition_Sync : FFragmentTestBase
+{
+	FObserverChangingComposition_Sync()
+	{
+		OperationObserved = EMassObservedOperation::Add;
+	}
+
+	virtual bool InstantTest() override
+	{
+		constexpr int32 EntitiesToSpawn = 3;
+		const FMassArchetypeHandle OriginalArchetype = FloatsArchetype;
+
+		FMassObserverManager& ObserverManager = EntityManager->GetObserverManager();
+		ObserverManager.AddObserverInstance(*FFragmentStruct::StaticStruct(), OperationObserved, *ObserverProcessor);
+
+		{
+			AITEST_SCOPED_CHECK("Use asynchronous API instead", 1);
+
+			ObserverProcessor->ForEachEntityChunkExecutionFunction = [EntityManager = EntityManager, &OriginalArchetype](FMassExecutionContext& Context)
+			{
+				EntityManager->AddFragmentToEntity(Context.GetEntity(0), FTestFragment_Int::StaticStruct());
+			};
+
+			EntityManager->BatchCreateEntities(OriginalArchetype, EntitiesToSpawn, EntitiesInt);
+
+			AITEST_EQUAL("Number of entities in the original archetype, no moves expected", EntityManager->DebugGetArchetypeEntitiesCount(OriginalArchetype), EntitiesToSpawn);
+		}
+		{
+			AITEST_SCOPED_CHECK("Use asynchronous API instead", 1);
+
+			ObserverProcessor->ForEachEntityChunkExecutionFunction = [EntityManager = EntityManager, &OriginalArchetype](FMassExecutionContext& Context)
+			{
+				FMassArchetypeEntityCollection EntityCollection(OriginalArchetype, Context.GetEntities(), FMassArchetypeEntityCollection::NoDuplicates);
+				EntityManager->BatchChangeFragmentCompositionForEntities(MakeArrayView(&EntityCollection, 1)
+					, FMassFragmentBitSet(*FTestFragment_Int::StaticStruct()), {});
+			};
+
+			EntityManager->BatchCreateEntities(OriginalArchetype, EntitiesToSpawn, EntitiesInt);
+
+			AITEST_EQUAL("Number of entities in the original archetype, no moves expected", EntityManager->DebugGetArchetypeEntitiesCount(OriginalArchetype), EntitiesToSpawn * 2);
+		}
+
+		return true;
+	}
+};
+IMPLEMENT_AI_INSTANT_TEST(FObserverChangingComposition_Sync, "System.Mass.Observer.ChangingCompositionSync");
+
+struct FObserverChangingComposition_Deferred : FFragmentTestBase
+{
+	FObserverChangingComposition_Deferred()
+	{
+		OperationObserved = EMassObservedOperation::Add;
+	}
+
+	virtual bool InstantTest() override
+	{
+		constexpr int32 EntitiesToSpawn = 3;
+		const FMassArchetypeHandle OriginalArchetype = FloatsArchetype;
+
+		ObserverProcessor->ForEachEntityChunkExecutionFunction = [](FMassExecutionContext& Context)
+			{
+				Context.Defer().PushCommand<FMassCommandAddFragments<FTestFragment_Int>>(Context.GetEntities());
+			};
+
+		FMassObserverManager& ObserverManager = EntityManager->GetObserverManager();
+		ObserverManager.AddObserverInstance(*FFragmentStruct::StaticStruct(), OperationObserved, *ObserverProcessor);
+
+		EntityManager->BatchCreateEntities(OriginalArchetype, EntitiesToSpawn, EntitiesInt);
+
+		AITEST_EQUAL("Number of entities in the original archetype", EntityManager->DebugGetArchetypeEntitiesCount(OriginalArchetype), 0);
+		AITEST_EQUAL("Number of entities in the target archetype", EntityManager->DebugGetArchetypeEntitiesCount(FloatsIntsArchetype), EntitiesToSpawn);
+
+		return true;
+	}
+};
+IMPLEMENT_AI_INSTANT_TEST(FObserverChangingComposition_Deferred, "System.Mass.Observer.ChangingCompositionDeferred");
+
 /**
  * This test aims to verify expected behavior of observers when there's a creation context active, when composition-mutating
  * operations are affecting entities other than the ones being created.
@@ -772,7 +849,69 @@ struct FMoveToAnotherArchetype_SingleEntity : FEntityTestBase
 		return true;
 	}
 };
-IMPLEMENT_AI_INSTANT_TEST(FMoveToAnotherArchetype_SingleEntity, "System.Mass.Observer.Move.UnobservedCompositionChange");
+IMPLEMENT_AI_INSTANT_TEST(FMoveToAnotherArchetype_SingleEntity, "System.Mass.Observer.MoveToAnotherArchetype");
+//
+//struct FRecursiveObserver : FEntityTestBase
+//{
+//	bool bTagAdded = false;
+//	bool bTagRemoved = false;
+//	
+//	virtual bool SetUp() override
+//	{
+//		if (FEntityTestBase::SetUp() == false)
+//		{
+//			return false;
+//		}
+//
+//		FMassObserverManager& ObserverManager = EntityManager->GetObserverManager();
+//
+//		auto CreateObserver = [this](const TFunction<void(FMassExecutionContext& Context)>& StoreResultFunction)
+//		{
+//			UMassTestProcessorBase* ObserverProcessor = NewTestProcessor<UMassTestProcessorBase>(EntityManager);
+//			ObserverProcessor->EntityQuery.AddRequirement<FTestFragment_Int>(EMassFragmentAccess::ReadOnly);
+//			ObserverProcessor->ForEachEntityChunkExecutionFunction = StoreResultFunction;
+//			return ObserverProcessor;
+//		};
+//
+//		ObserverManager.AddObserverInstance(*FTestTag_A::StaticStruct(), EMassObservedOperation::Add
+//			, *CreateObserver([this](FMassExecutionContext& Context)
+//				{
+//					bTagAdded = true;
+//				}
+//			));
+//
+//		ObserverManager.AddObserverInstance(*FTestTag_A::StaticStruct(), EMassObservedOperation::Remove
+//			, *CreateObserver([this](FMassExecutionContext& Context)
+//				{
+//					bTagRemoved = true;
+//				}
+//			));
+//
+//		return true;
+//	}
+//
+//	virtual bool InstantTest() override
+//	{
+//		const FMassEntityHandle EntityHandle = EntityManager->CreateEntity(IntsArchetype);
+//
+//		// create target archetype
+//		const FMassArchetypeHandle TargetArchetypeHandle = EntityManager->CreateArchetype(IntsArchetype, { FTestTag_A::StaticStruct(), FTestFragment_Float::StaticStruct() });
+//
+//		EntityManager->MoveEntityToAnotherArchetype(EntityHandle, TargetArchetypeHandle);
+//		AITEST_TRUE("Tag addition observer has been executed", bTagAdded);
+//		AITEST_TRUE("Fragment addition observer has been executed", bFloatAdded);
+//		AITEST_FALSE("(NOT) Tag removal observer has been executed", bTagRemoved);
+//		AITEST_FALSE("(NOT) Fragment addition observer has been executed", bFloatRemoved);
+//
+//		// moving back to the original archetype will remove the two added elements, and should trigger observers
+//		EntityManager->MoveEntityToAnotherArchetype(EntityHandle, IntsArchetype);
+//		AITEST_TRUE("Tag removal observer has been executed", bTagRemoved);
+//		AITEST_TRUE("Fragment addition observer has been executed", bFloatRemoved);
+//
+//		return true;
+//	}
+//};
+//IMPLEMENT_AI_INSTANT_TEST(FMoveToAnotherArchetype_SingleEntity, "System.Mass.Observer.Recursive");
 
 } // UE::Mass::Test::Observers
 
diff --git a/Engine/Source/Runtime/MassEntity/Private/MassEntityManager.cpp b/Engine/Source/Runtime/MassEntity/Private/MassEntityManager.cpp
--- a/Engine/Source/Runtime/MassEntity/Private/MassEntityManager.cpp
+++ b/Engine/Source/Runtime/MassEntity/Private/MassEntityManager.cpp
@@ -17,7 +17,8 @@
 #include "MassEntityBuilder.h"
 #include "MassTypeManager.h"
 
-#define CHECK_SYNC_API() checkf(IsProcessing() == false, TEXT("Synchronous API function %hs called during mass processing. Use asynchronous API instead."), __FUNCTION__)
+#define CHECK_SYNC_API() testableCheckf(IsProcessing() == false, TEXT("Synchronous API function %hs called during mass processing. Use asynchronous API instead."), __FUNCTION__)
+#define CHECK_SYNC_API_RETURN(ReturnValue) testableCheckfReturn(IsProcessing() == false, return ReturnValue, TEXT("Synchronous API function %hs called during mass processing. Use asynchronous API instead."), __FUNCTION__)
 
 const FMassEntityHandle FMassEntityManager::InvalidEntity;
 
@@ -660,7 +661,7 @@ void FMassEntityManager::DoEntityCompaction(const double TimeAllowed)
 
 FMassEntityHandle FMassEntityManager::CreateEntity(const FMassArchetypeHandle& ArchetypeHandle, const FMassArchetypeSharedFragmentValues& SharedFragmentValues)
 {
-	CHECK_SYNC_API();
+	CHECK_SYNC_API_RETURN({});
 	check(ArchetypeHandle.IsValid());
 
 	const FMassEntityHandle Entity = ReserveEntity();
@@ -673,7 +674,7 @@ FMassEntityHandle FMassEntityManager::CreateEntity(const FMassArchetypeHandle& A
 
 FMassEntityHandle FMassEntityManager::CreateEntity(TConstArrayView<FInstancedStruct> FragmentInstanceList, const FMassArchetypeSharedFragmentValues& SharedFragmentValues, const FMassArchetypeCreationParams& CreationParams)
 {
-	CHECK_SYNC_API();
+	CHECK_SYNC_API_RETURN({});
 	check(FragmentInstanceList.Num() > 0);
 
 	const FMassArchetypeHandle& ArchetypeHandle = CreateArchetype(FMassArchetypeCompositionDescriptor(FragmentInstanceList,
@@ -766,7 +767,7 @@ int32 FMassEntityManager::BatchReserveEntities(TArrayView<FMassEntityHandle> InO
 TSharedRef<FMassEntityManager::FEntityCreationContext> FMassEntityManager::BatchBuildEntities(const FMassArchetypeEntityCollectionWithPayload& EncodedEntitiesWithPayload
 	, const FMassFragmentBitSet& FragmentsAffected, const FMassArchetypeSharedFragmentValues& SharedFragmentValues, const FMassArchetypeCreationParams& CreationParams)
 {
-	CHECK_SYNC_API();
+	CHECK_SYNC_API_RETURN(MakeShareable(new FEntityCreationContext));
 	check(SharedFragmentValues.IsSorted());
 
 	FMassArchetypeCompositionDescriptor Composition(FragmentsAffected, FMassTagBitSet(), FMassChunkFragmentBitSet(), FMassSharedFragmentBitSet(), FMassConstSharedFragmentBitSet());
@@ -786,7 +787,7 @@ TSharedRef<FMassEntityManager::FEntityCreationContext> FMassEntityManager::Batch
 	, FMassArchetypeCompositionDescriptor&& Composition
 	, const FMassArchetypeSharedFragmentValues& SharedFragmentValues, const FMassArchetypeCreationParams& CreationParams)
 {
-	CHECK_SYNC_API();
+	CHECK_SYNC_API_RETURN(MakeShareable(new FEntityCreationContext));
 
 	TRACE_CPUPROFILER_EVENT_SCOPE(Mass_BatchBuildEntities);
 
@@ -843,7 +844,7 @@ TSharedRef<FMassEntityManager::FEntityCreationContext> FMassEntityManager::Batch
 TSharedRef<FMassEntityManager::FEntityCreationContext> FMassEntityManager::BatchCreateReservedEntities(const FMassArchetypeHandle& ArchetypeHandle
 	, const FMassArchetypeSharedFragmentValues& SharedFragmentValues, TConstArrayView<FMassEntityHandle> ReservedEntities)
 {
-	CHECK_SYNC_API();
+	CHECK_SYNC_API_RETURN(MakeShareable(new FEntityCreationContext));
 	checkf(!ReservedEntities.IsEmpty(), TEXT("No reserved entities given to batch create."));
 
 	TRACE_CPUPROFILER_EVENT_SCOPE(Mass_BatchCreateReservedEntities);
@@ -856,7 +857,7 @@ TSharedRef<FMassEntityManager::FEntityCreationContext> FMassEntityManager::Batch
 TSharedRef<FMassEntityManager::FEntityCreationContext> FMassEntityManager::BatchCreateEntities(const FMassArchetypeHandle& ArchetypeHandle
 	, const FMassArchetypeSharedFragmentValues& SharedFragmentValues, const int32 Count, TArray<FMassEntityHandle>& InOutEntities)
 {
-	CHECK_SYNC_API();
+	CHECK_SYNC_API_RETURN(MakeShareable(new FEntityCreationContext));
 	testableCheckfReturn(ArchetypeHandle.IsValid(), return GetOrMakeCreationContext()
 		, TEXT("%hs expecting a valid ArchetypeHandle"), __FUNCTION__);
 
@@ -1454,7 +1455,7 @@ void FMassEntityManager::RemoveTagFromEntity(FMassEntityHandle Entity, const USc
 
 bool FMassEntityManager::AddConstSharedFragmentToEntity(const FMassEntityHandle Entity, const FConstSharedStruct& InConstSharedFragment)
 {
-	CHECK_SYNC_API();
+	CHECK_SYNC_API_RETURN(false);
 
 	if (!ensureMsgf(InConstSharedFragment.IsValid(), TEXT("%hs parameter Fragment is expected to be valid"), __FUNCTION__))
 	{
@@ -1504,7 +1505,7 @@ bool FMassEntityManager::AddConstSharedFragmentToEntity(const FMassEntityHandle
 
 bool FMassEntityManager::RemoveConstSharedFragmentFromEntity(const FMassEntityHandle Entity, const UScriptStruct& ConstSharedFragmentType)
 {
-	CHECK_SYNC_API();
+	CHECK_SYNC_API_RETURN(false);
 
 	if (!ensureMsgf(UE::Mass::IsA<FMassConstSharedFragment>(&ConstSharedFragmentType), TEXT("%hs parameter ConstSharedFragmentType is expected to be a FMassConstSharedFragment"), __FUNCTION__))
 	{
@@ -1547,7 +1548,7 @@ bool FMassEntityManager::RemoveConstSharedFragmentFromEntity(const FMassEntityHa
 
 bool FMassEntityManager::AddSharedFragmentToEntity(const FMassEntityHandle Entity, const FSharedStruct& InSharedFragment)
 {
-	CHECK_SYNC_API();
+	CHECK_SYNC_API_RETURN(false);
 
 	if (!ensureMsgf(InSharedFragment.IsValid(), TEXT("%hs parameter Fragment is expected to be valid"), __FUNCTION__))
 	{
@@ -1597,7 +1598,7 @@ bool FMassEntityManager::AddSharedFragmentToEntity(const FMassEntityHandle Entit
 
 bool FMassEntityManager::RemoveSharedFragmentFromEntity(const FMassEntityHandle Entity, const UScriptStruct& SharedFragmentType)
 {
-	CHECK_SYNC_API();
+	CHECK_SYNC_API_RETURN(false);
 
 	if (!ensureMsgf(UE::Mass::IsA<FMassSharedFragment>(&SharedFragmentType), TEXT("%hs parameter SharedFragmentType is expected to be a FMassSharedFragment"), __FUNCTION__))
 	{
@@ -2350,6 +2351,12 @@ const UE::Mass::FStorageType& FMassEntityManager::DebugGetEntityStorageInterface
 {
 	return GetEntityStorageInterface();
 }
+
+bool FMassEntityManager::DebugHasCommandsToFlush() const
+{
+	checkfSlow(NumCommandBuffers == 2, TEXT("This check relies on there being two command buffers."));
+	return DeferredCommandBuffers[0]->HasPendingCommands() || DeferredCommandBuffers[1]->HasPendingCommands();
+}
 #endif // WITH_MASSENTITY_DEBUG
 
 //-----------------------------------------------------------------------------
diff --git a/Engine/Source/Runtime/MassEntity/Public/MassEntityManager.h b/Engine/Source/Runtime/MassEntity/Public/MassEntityManager.h
--- a/Engine/Source/Runtime/MassEntity/Public/MassEntityManager.h
+++ b/Engine/Source/Runtime/MassEntity/Public/MassEntityManager.h
@@ -1059,6 +1059,8 @@ PRAGMA_ENABLE_DEPRECATION_WARNINGS
 	UE::Mass::FStorageType& DebugGetEntityStorageInterface();
 	// For use by the friend MassDebugger
 	const UE::Mass::FStorageType& DebugGetEntityStorageInterface() const;
+
+	bool DebugHasCommandsToFlush() const;
 #endif // WITH_MASSENTITY_DEBUG
 
 protected:
@@ -1204,7 +1206,8 @@ private:
 	 * by MassObservers reacting to the commands being flushed currently).
 	 */
 	uint8 OpenedCommandBufferIndex = 0;
-	TStaticArray<TSharedPtr<FMassCommandBuffer>, 2> DeferredCommandBuffers;
+	static constexpr int32 NumCommandBuffers = 2;
+	TStaticArray<TSharedPtr<FMassCommandBuffer>, NumCommandBuffers> DeferredCommandBuffers;
 	
 	std::atomic<int32> ProcessingScopeCount = 0;
 
