[Mass] unit tests CIS fix

#rnx

diff --git a/Engine/Source/Developer/MassEntityTestSuite/Private/MassObserverTest.cpp b/Engine/Source/Developer/MassEntityTestSuite/Private/MassObserverTest.cpp
--- a/Engine/Source/Developer/MassEntityTestSuite/Private/MassObserverTest.cpp
+++ b/Engine/Source/Developer/MassEntityTestSuite/Private/MassObserverTest.cpp
@@ -625,6 +625,7 @@ struct FFragmentTest_EntityCreation_Individuals : FFragmentTestBase
 };
 IMPLEMENT_AI_INSTANT_TEST(FFragmentTest_EntityCreation_Individuals, "System.Mass.Observer.Create.FragmentIndividualEntities");
 
+#if WITH_MASSENTITY_DEBUG
 struct FObserverChangingComposition_Sync : FFragmentTestBase
 {
 	FObserverChangingComposition_Sync()
@@ -701,6 +702,7 @@ struct FObserverChangingComposition_Deferred : FFragmentTestBase
 	}
 };
 IMPLEMENT_AI_INSTANT_TEST(FObserverChangingComposition_Deferred, "System.Mass.Observer.ChangingCompositionDeferred");
+#endif // WITH_MASSENTITY_DEBUG
 
 /**
  * This test aims to verify expected behavior of observers when there's a creation context active, when composition-mutating
