[Mass] warning V1047 in MassObserverTests.cpp fixed

#rnx

diff --git a/Engine/Source/Developer/MassEntityTestSuite/Private/MassObserverTest.cpp b/Engine/Source/Developer/MassEntityTestSuite/Private/MassObserverTest.cpp
--- a/Engine/Source/Developer/MassEntityTestSuite/Private/MassObserverTest.cpp
+++ b/Engine/Source/Developer/MassEntityTestSuite/Private/MassObserverTest.cpp
@@ -644,7 +644,7 @@ struct FObserverChangingComposition_Sync : FFragmentTestBase
 		{
 			AITEST_SCOPED_CHECK("Use asynchronous API instead", 1);
 
-			ObserverProcessor->ForEachEntityChunkExecutionFunction = [EntityManager = EntityManager, &OriginalArchetype](FMassExecutionContext& Context)
+			ObserverProcessor->ForEachEntityChunkExecutionFunction = [EntityManager = EntityManager](FMassExecutionContext& Context)
 			{
 				EntityManager->AddFragmentToEntity(Context.GetEntity(0), FTestFragment_Int::StaticStruct());
 			};
@@ -656,7 +656,7 @@ struct FObserverChangingComposition_Sync : FFragmentTestBase
 		{
 			AITEST_SCOPED_CHECK("Use asynchronous API instead", 1);
 
-			ObserverProcessor->ForEachEntityChunkExecutionFunction = [EntityManager = EntityManager, &OriginalArchetype](FMassExecutionContext& Context)
+			ObserverProcessor->ForEachEntityChunkExecutionFunction = [EntityManager = EntityManager, OriginalArchetype](FMassExecutionContext& Context)
 			{
 				FMassArchetypeEntityCollection EntityCollection(OriginalArchetype, Context.GetEntities(), FMassArchetypeEntityCollection::NoDuplicates);
 				EntityManager->BatchChangeFragmentCompositionForEntities(MakeArrayView(&EntityCollection, 1)
