[Mass] Unit tests: renamed System.Mass.EntityCollection.* tests to System.Mass.ArchetypeEntityCollection.* due to a conflict with tests for the new UE::Mass::FEntityCollection type

#rnx

diff --git a/Engine/Source/Developer/MassEntityTestSuite/Private/MassChunkCollectionTest.cpp b/Engine/Source/Developer/MassEntityTestSuite/Private/MassChunkCollectionTest.cpp
--- a/Engine/Source/Developer/MassEntityTestSuite/Private/MassChunkCollectionTest.cpp
+++ b/Engine/Source/Developer/MassEntityTestSuite/Private/MassChunkCollectionTest.cpp
@@ -16,10 +16,10 @@ UE_DISABLE_OPTIMIZATION_SHIP
 //----------------------------------------------------------------------//
 // tests 
 //----------------------------------------------------------------------//
-namespace FMassEntityCollectionTest
+namespace FMassArchetypeEntityCollectionTest
 {
 
-struct FEntityCollectionTestBase : FEntityTestBase
+struct FArchetypeEntityCollectionTestBase : FEntityTestBase
 {
 	TArray<FMassEntityHandle> Entities;
 
@@ -39,7 +39,7 @@ struct FEntityCollectionTestBase : FEntityTestBase
 	}
 };
 
-struct FEntityCollection_CreateBasic : FEntityCollectionTestBase
+struct FArchetypeEntityCollection_CreateBasic : FArchetypeEntityCollectionTestBase
 {
 	virtual bool InstantTest() override
 	{
@@ -79,9 +79,9 @@ struct FEntityCollection_CreateBasic : FEntityCollectionTestBase
 		return true;
 	}
 };
-IMPLEMENT_AI_INSTANT_TEST(FEntityCollection_CreateBasic, "System.Mass.EntityCollection.Create.Basic");
+IMPLEMENT_AI_INSTANT_TEST(FArchetypeEntityCollection_CreateBasic, "System.Mass.ArchetypeEntityCollection.Create.Basic");
 
-struct FEntityCollection_CreateOrderInvariant : FEntityCollectionTestBase
+struct FArchetypeEntityCollection_CreateOrderInvariant : FArchetypeEntityCollectionTestBase
 {
 	virtual bool InstantTest() override
 	{
@@ -104,9 +104,9 @@ struct FEntityCollection_CreateOrderInvariant : FEntityCollectionTestBase
 		return true;
 	}
 };
-IMPLEMENT_AI_INSTANT_TEST(FEntityCollection_CreateOrderInvariant, "System.Mass.EntityCollection.Create.OrderInvariant");
+IMPLEMENT_AI_INSTANT_TEST(FArchetypeEntityCollection_CreateOrderInvariant, "System.Mass.ArchetypeEntityCollection.Create.OrderInvariant");
 
-struct FEntityCollection_CreateCrossChunk : FEntityTestBase
+struct FArchetypeEntityCollection_CreateCrossChunk : FEntityTestBase
 {
 	virtual bool InstantTest() override
 	{
@@ -135,9 +135,9 @@ struct FEntityCollection_CreateCrossChunk : FEntityTestBase
 		return true;
 	}
 };
-IMPLEMENT_AI_INSTANT_TEST(FEntityCollection_CreateCrossChunk, "System.Mass.EntityCollection.Create.CrossChunk");
+IMPLEMENT_AI_INSTANT_TEST(FArchetypeEntityCollection_CreateCrossChunk, "System.Mass.ArchetypeEntityCollection.Create.CrossChunk");
 
-struct FEntityCollection_CreateWithDuplicatesTrivial : FEntityCollectionTestBase
+struct FArchetypeEntityCollection_CreateWithDuplicatesTrivial : FArchetypeEntityCollectionTestBase
 {
 	virtual bool InstantTest() override
 	{
@@ -153,9 +153,9 @@ struct FEntityCollection_CreateWithDuplicatesTrivial : FEntityCollectionTestBase
 		return true;
 	}
 };
-IMPLEMENT_AI_INSTANT_TEST(FEntityCollection_CreateWithDuplicatesTrivial, "System.Mass.EntityCollection.Create.TrivialDuplicates");
+IMPLEMENT_AI_INSTANT_TEST(FArchetypeEntityCollection_CreateWithDuplicatesTrivial, "System.Mass.ArchetypeEntityCollection.Create.TrivialDuplicates");
 
-struct FEntityCollection_CreateWithDuplicates : FEntityCollectionTestBase
+struct FArchetypeEntityCollection_CreateWithDuplicates : FArchetypeEntityCollectionTestBase
 {
 	virtual bool InstantTest() override
 	{
@@ -175,9 +175,9 @@ struct FEntityCollection_CreateWithDuplicates : FEntityCollectionTestBase
 		return true;
 	}
 };
-IMPLEMENT_AI_INSTANT_TEST(FEntityCollection_CreateWithDuplicates, "System.Mass.EntityCollection.Create.Duplicates");
+IMPLEMENT_AI_INSTANT_TEST(FArchetypeEntityCollection_CreateWithDuplicates, "System.Mass.ArchetypeEntityCollection.Create.Duplicates");
 
-struct FEntityCollection_CreateWithInvalidDuplicates : FEntityCollectionTestBase
+struct FArchetypeEntityCollection_CreateWithInvalidDuplicates : FArchetypeEntityCollectionTestBase
 {
 	virtual bool InstantTest() override
 	{
@@ -219,9 +219,9 @@ struct FEntityCollection_CreateWithInvalidDuplicates : FEntityCollectionTestBase
 		return true;
 	}
 };
-IMPLEMENT_AI_INSTANT_TEST(FEntityCollection_CreateWithInvalidDuplicates, "System.Mass.EntityCollection.Create.InvalidDuplicates");
+IMPLEMENT_AI_INSTANT_TEST(FArchetypeEntityCollection_CreateWithInvalidDuplicates, "System.Mass.ArchetypeEntityCollection.Create.InvalidDuplicates");
 
-struct FEntityCollection_CreateWithInvalidDuplicatesWithPayload : FEntityCollectionTestBase
+struct FArchetypeEntityCollection_CreateWithInvalidDuplicatesWithPayload : FArchetypeEntityCollectionTestBase
 {
 	virtual bool InstantTest() override
 	{
@@ -257,10 +257,10 @@ struct FEntityCollection_CreateWithInvalidDuplicatesWithPayload : FEntityCollect
 		return true;
 	}
 };
-IMPLEMENT_AI_INSTANT_TEST(FEntityCollection_CreateWithInvalidDuplicatesWithPayload, "System.Mass.EntityCollection.Create.InvalidDuplicatesWithPayloadWithPayload");
+IMPLEMENT_AI_INSTANT_TEST(FArchetypeEntityCollection_CreateWithInvalidDuplicatesWithPayload, "System.Mass.ArchetypeEntityCollection.Create.InvalidDuplicatesWithPayloadWithPayload");
 
 #if WITH_MASSENTITY_DEBUG
-struct FEntityCollection_WithPayloadBase : FEntityCollectionTestBase
+struct FArchetypeEntityCollection_WithPayloadBase : FArchetypeEntityCollectionTestBase
 {
 	virtual bool SetUp() override
 	{
@@ -273,7 +273,7 @@ struct FEntityCollection_WithPayloadBase : FEntityCollectionTestBase
 	}
 };
 
-struct FEntityCollection_TrivialDuplicatesWithPayload : FEntityCollection_WithPayloadBase
+struct FArchetypeEntityCollection_TrivialDuplicatesWithPayload : FArchetypeEntityCollection_WithPayloadBase
 {
 	virtual bool InstantTest() override
 	{
@@ -301,10 +301,10 @@ struct FEntityCollection_TrivialDuplicatesWithPayload : FEntityCollection_WithPa
 		return true;
 	}
 };
-IMPLEMENT_AI_INSTANT_TEST(FEntityCollection_TrivialDuplicatesWithPayload, "System.Mass.EntityCollection.Create.TrivialDuplicatesWithPayload");
+IMPLEMENT_AI_INSTANT_TEST(FArchetypeEntityCollection_TrivialDuplicatesWithPayload, "System.Mass.ArchetypeEntityCollection.Create.TrivialDuplicatesWithPayload");
 
 // @todo also add another archetype
-struct FEntityCollection_MultiDuplicatesWithPayload : FEntityCollection_WithPayloadBase
+struct FArchetypeEntityCollection_MultiDuplicatesWithPayload : FArchetypeEntityCollection_WithPayloadBase
 {
 	virtual bool InstantTest() override
 	{
@@ -344,9 +344,9 @@ struct FEntityCollection_MultiDuplicatesWithPayload : FEntityCollection_WithPayl
 		return true;
 	}
 };
-IMPLEMENT_AI_INSTANT_TEST(FEntityCollection_MultiDuplicatesWithPayload, "System.Mass.EntityCollection.Create.MultiDuplicatesWithPayload");
+IMPLEMENT_AI_INSTANT_TEST(FArchetypeEntityCollection_MultiDuplicatesWithPayload, "System.Mass.ArchetypeEntityCollection.Create.MultiDuplicatesWithPayload");
 
-struct FEntityCollection_WithPayload : FEntityCollection_WithPayloadBase
+struct FArchetypeEntityCollection_WithPayload : FArchetypeEntityCollection_WithPayloadBase
 {
 	virtual bool InstantTest() override
 	{
@@ -388,9 +388,9 @@ struct FEntityCollection_WithPayload : FEntityCollection_WithPayloadBase
 		return true;
 	}
 };
-IMPLEMENT_AI_INSTANT_TEST(FEntityCollection_WithPayload, "System.Mass.EntityCollection.Create.WithPayload");
+IMPLEMENT_AI_INSTANT_TEST(FArchetypeEntityCollection_WithPayload, "System.Mass.ArchetypeEntityCollection.Create.WithPayload");
 
-struct FEntityCollection_ZeroLength_Destroy : FEntityTestBase
+struct FArchetypeEntityCollection_ZeroLength_Destroy : FEntityTestBase
 {
 	virtual bool InstantTest() override
 	{
@@ -416,9 +416,9 @@ struct FEntityCollection_ZeroLength_Destroy : FEntityTestBase
 		return true;
 	}
 };
-IMPLEMENT_AI_INSTANT_TEST(FEntityCollection_ZeroLength_Destroy, "System.Mass.Entity.ZeroLengthRanges.Destroy");
+IMPLEMENT_AI_INSTANT_TEST(FArchetypeEntityCollection_ZeroLength_Destroy, "System.Mass.Entity.ZeroLengthRanges.Destroy");
 
-struct FEntityCollection_ZeroLength_Execute : FEntityTestBase
+struct FArchetypeEntityCollection_ZeroLength_Execute : FEntityTestBase
 {
 	virtual bool InstantTest() override
 	{
@@ -450,9 +450,9 @@ struct FEntityCollection_ZeroLength_Execute : FEntityTestBase
 		return true;
 	}
 };
-IMPLEMENT_AI_INSTANT_TEST(FEntityCollection_ZeroLength_Execute, "System.Mass.Entity.ZeroLengthRanges.Execute");
+IMPLEMENT_AI_INSTANT_TEST(FArchetypeEntityCollection_ZeroLength_Execute, "System.Mass.Entity.ZeroLengthRanges.Execute");
 
-struct FEntityCollection_ZeroLength_BatchMove : FEntityTestBase
+struct FArchetypeEntityCollection_ZeroLength_BatchMove : FEntityTestBase
 {
 	virtual bool InstantTest() override
 	{
@@ -482,9 +482,9 @@ struct FEntityCollection_ZeroLength_BatchMove : FEntityTestBase
 		return true;
 	}
 };
-IMPLEMENT_AI_INSTANT_TEST(FEntityCollection_ZeroLength_BatchMove, "System.Mass.Entity.ZeroLengthRanges.BatchMove");
+IMPLEMENT_AI_INSTANT_TEST(FArchetypeEntityCollection_ZeroLength_BatchMove, "System.Mass.Entity.ZeroLengthRanges.BatchMove");
 
-struct FEntityCollection_ZeroLength_BatchSetFragmentValues : FEntityTestBase
+struct FArchetypeEntityCollection_ZeroLength_BatchSetFragmentValues : FEntityTestBase
 {
 	virtual bool InstantTest() override
 	{
@@ -499,8 +499,7 @@ struct FEntityCollection_ZeroLength_BatchSetFragmentValues : FEntityTestBase
 		AITEST_EQUAL("Entity creation generated expected number of entities", EntityCount, FMassArchetypeHelper::ArchetypeDataFromHandleChecked(SourceArchetype).GetNumEntities());
 		check(EntityCount == Entities.Num());
 
-		FMassArchetypeEntityCollection Collection(SourceArchetype);
-		FMassArchetypeEntityCollectionWithPayload EntityCollectionWithEmptyPayload(Collection);
+		FMassArchetypeEntityCollectionWithPayload EntityCollectionWithEmptyPayload{FMassArchetypeEntityCollection(SourceArchetype)};
 
 		const FMassFragmentBitSet AddedFragmentsBitSet(*FTestFragment_Int::StaticStruct());
 		EntityManager->BatchAddFragmentInstancesForEntities(MakeArrayView(&EntityCollectionWithEmptyPayload, 1), AddedFragmentsBitSet);
@@ -515,9 +514,9 @@ struct FEntityCollection_ZeroLength_BatchSetFragmentValues : FEntityTestBase
 		return true;
 	}
 };
-IMPLEMENT_AI_INSTANT_TEST(FEntityCollection_ZeroLength_BatchSetFragmentValues, "System.Mass.Entity.ZeroLengthRanges.BatchSetFragmentValues");
+IMPLEMENT_AI_INSTANT_TEST(FArchetypeEntityCollection_ZeroLength_BatchSetFragmentValues, "System.Mass.Entity.ZeroLengthRanges.BatchSetFragmentValues");
 
-struct FEntityCollection_ExportHandles_Created : FEntityTestBase
+struct FArchetypeEntityCollection_ExportHandles_Created : FEntityTestBase
 {
 	virtual bool InstantTest() override
 	{
@@ -547,7 +546,7 @@ struct FEntityCollection_ExportHandles_Created : FEntityTestBase
 		return true;
 	}
 };
-IMPLEMENT_AI_INSTANT_TEST(FEntityCollection_ExportHandles_Created, "System.Mass.EntityCollection.ExportHandles.CreatedEntities");
+IMPLEMENT_AI_INSTANT_TEST(FArchetypeEntityCollection_ExportHandles_Created, "System.Mass.ArchetypeEntityCollection.ExportHandles.CreatedEntities");
 
 #endif // WITH_MASSENTITY_DEBUG
 //
