[Mass] SA Issue 835337: modified the operation in entity collection's unit tests PVS complained about

#rnx

diff --git a/Engine/Source/Developer/MassEntityTestSuite/Private/MassEntityCollectionTest.cpp b/Engine/Source/Developer/MassEntityTestSuite/Private/MassEntityCollectionTest.cpp
--- a/Engine/Source/Developer/MassEntityTestSuite/Private/MassEntityCollectionTest.cpp
+++ b/Engine/Source/Developer/MassEntityTestSuite/Private/MassEntityCollectionTest.cpp
@@ -43,7 +43,7 @@ struct FEntityCollectionTestBase : FEntityTestBase
 		while (EntitiesSubSet.Num() < NumEntities)
 		{
 			EntitiesSubSet.AddUnique(GetEntityArray(ArrayIndex)[RandomStream.RandRange(0, EntitiesToCreatePerArchetype - 1)]);
-			ArrayIndex = (++ArrayIndex % NumArchetypesUsed);
+			ArrayIndex = (ArrayIndex + 1) % NumArchetypesUsed;
 		}
 
 		return EntitiesSubSet;
