[StructUtils] make InstancedStruct, SharedStruct and StructViews nodiscard type

diff --git a/Engine/Plugins/Runtime/MassGameplay/Source/MassSpawner/Public/MassEntityTemplate.h b/Engine/Plugins/Runtime/MassGameplay/Source/MassSpawner/Public/MassEntityTemplate.h
--- a/Engine/Plugins/Runtime/MassGameplay/Source/MassSpawner/Public/MassEntityTemplate.h
+++ b/Engine/Plugins/Runtime/MassGameplay/Source/MassSpawner/Public/MassEntityTemplate.h
@@ -224,7 +224,7 @@ struct MASSSPAWNER_API FMassEntityTemplateData
 			if (!Composition.ConstSharedFragments.Contains(*FragmentType))
 			{
 				Composition.ConstSharedFragments.Add(*FragmentType);
-				SharedFragmentValues.AddConstSharedFragment(SharedFragment);
+				SharedFragmentValues.Add(SharedFragment);
 			}
 #if DO_ENSURE
 			else
@@ -245,7 +245,7 @@ struct MASSSPAWNER_API FMassEntityTemplateData
 			if (!Composition.SharedFragments.Contains(*FragmentType))
 			{
 				Composition.SharedFragments.Add(*FragmentType);
-				SharedFragmentValues.AddSharedFragment(SharedFragment);
+				SharedFragmentValues.Add(SharedFragment);
 			}
 	#if DO_ENSURE
 			else
diff --git a/Engine/Plugins/Runtime/MassGameplay/Source/MassSpawner/Public/MassEntityTemplateRegistry.h b/Engine/Plugins/Runtime/MassGameplay/Source/MassSpawner/Public/MassEntityTemplateRegistry.h
--- a/Engine/Plugins/Runtime/MassGameplay/Source/MassSpawner/Public/MassEntityTemplateRegistry.h
+++ b/Engine/Plugins/Runtime/MassGameplay/Source/MassSpawner/Public/MassEntityTemplateRegistry.h
@@ -121,14 +121,14 @@ struct FMassEntityTemplateBuildContext
 
 	void AddConstSharedFragment(const FConstSharedStruct& InSharedFragment)
 	{
-		checkf(InSharedFragment.GetScriptStruct(), TEXT("Expecting a valide shared fragment type"));
+		checkf(InSharedFragment.GetScriptStruct(), TEXT("Expecting a valid shared fragment type"));
 		TypeAdded(*InSharedFragment.GetScriptStruct());
 		TemplateData.AddConstSharedFragment(InSharedFragment);
 	}
 
 	void AddSharedFragment(const FSharedStruct& InSharedFragment)
 	{
-		checkf(InSharedFragment.GetScriptStruct(), TEXT("Expecting a valide shared fragment type"));
+		checkf(InSharedFragment.GetScriptStruct(), TEXT("Expecting a valid shared fragment type"));
 		TypeAdded(*InSharedFragment.GetScriptStruct());
 		TemplateData.AddSharedFragment(InSharedFragment);
 	}
diff --git a/Engine/Source/Developer/MassEntityTestSuite/Private/MassQueryTest.cpp b/Engine/Source/Developer/MassEntityTestSuite/Private/MassQueryTest.cpp
--- a/Engine/Source/Developer/MassEntityTestSuite/Private/MassQueryTest.cpp
+++ b/Engine/Source/Developer/MassEntityTestSuite/Private/MassQueryTest.cpp
@@ -624,7 +624,7 @@ struct FQueryTest_AllOptional : FEntityTestBase
 			FTestSharedFragment_Int FragmentInstance;
 			FSharedStruct SharedFragmentInstance = FSharedStruct::Make(FragmentInstance);
 			FMassArchetypeSharedFragmentValues SharedFragmentValues;
-			SharedFragmentValues.AddSharedFragment(SharedFragmentInstance);
+			SharedFragmentValues.Add(SharedFragmentInstance);
 
 			FMassArchetypeEntityCollection Collection(IntsArchetype, MakeArrayView(&Entities[CurrentEntityIndex++], 1), FMassArchetypeEntityCollection::EDuplicatesHandling::NoDuplicates);
 			EntityManager->BatchAddSharedFragmentsForEntities(MakeArrayView(&Collection, 1), SharedFragmentValues);
@@ -634,7 +634,7 @@ struct FQueryTest_AllOptional : FEntityTestBase
 			FTestSharedFragment_Float FragmentInstance;
 			FSharedStruct SharedFragmentInstance = FSharedStruct::Make(FragmentInstance);
 			FMassArchetypeSharedFragmentValues SharedFragmentValues;
-			SharedFragmentValues.AddSharedFragment(SharedFragmentInstance);
+			SharedFragmentValues.Add(SharedFragmentInstance);
 
 			FMassArchetypeEntityCollection Collection(IntsArchetype, MakeArrayView(&Entities[CurrentEntityIndex++], 1), FMassArchetypeEntityCollection::EDuplicatesHandling::NoDuplicates);
 			EntityManager->BatchAddSharedFragmentsForEntities(MakeArrayView(&Collection, 1), SharedFragmentValues);
@@ -786,7 +786,7 @@ struct FQueryTest_JustASharedFragment : FEntityTestBase
 			FTestSharedFragment_Int FragmentInstance;
 			FSharedStruct SharedFragmentInstance = FSharedStruct::Make(FragmentInstance);
 			FMassArchetypeSharedFragmentValues SharedFragmentValues;
-			SharedFragmentValues.AddSharedFragment(SharedFragmentInstance);
+			SharedFragmentValues.Add(SharedFragmentInstance);
 
 			FMassArchetypeEntityCollection Collection(IntsArchetype, MakeArrayView(&Entities[CurrentEntityIndex++], 1), FMassArchetypeEntityCollection::EDuplicatesHandling::NoDuplicates);
 			EntityManager->BatchAddSharedFragmentsForEntities(MakeArrayView(&Collection, 1), SharedFragmentValues);
@@ -796,7 +796,7 @@ struct FQueryTest_JustASharedFragment : FEntityTestBase
 			FTestSharedFragment_Float FragmentInstance;
 			FSharedStruct SharedFragmentInstance = FSharedStruct::Make(FragmentInstance);
 			FMassArchetypeSharedFragmentValues SharedFragmentValues;
-			SharedFragmentValues.AddSharedFragment(SharedFragmentInstance);
+			SharedFragmentValues.Add(SharedFragmentInstance);
 
 			FMassArchetypeEntityCollection Collection(IntsArchetype, MakeArrayView(&Entities[CurrentEntityIndex++], 1), FMassArchetypeEntityCollection::EDuplicatesHandling::NoDuplicates);
 			EntityManager->BatchAddSharedFragmentsForEntities(MakeArrayView(&Collection, 1), SharedFragmentValues);
diff --git a/Engine/Source/Developer/MassEntityTestSuite/Private/MassSharedFragmentsTest.cpp b/Engine/Source/Developer/MassEntityTestSuite/Private/MassSharedFragmentsTest.cpp
--- a/Engine/Source/Developer/MassEntityTestSuite/Private/MassSharedFragmentsTest.cpp
+++ b/Engine/Source/Developer/MassEntityTestSuite/Private/MassSharedFragmentsTest.cpp
@@ -44,7 +44,7 @@ struct FSharedFragmentValues_Create : FExecutionTestBase
 		{
 			FTestSharedFragment_Int FragmentInstance(TestIntValue);
 			FSharedStruct SharedFragmentInstance = FSharedStruct::Make(FragmentInstance);
-			Values.AddSharedFragment(SharedFragmentInstance);
+			Values.Add(SharedFragmentInstance);
 		}
 
 		const FTestSharedFragment_Int* ConstInstance = GetConstSharedFragmentPtr<FTestSharedFragment_Int>(Values);
@@ -68,7 +68,7 @@ struct FSharedFragmentValues_CreateConst : FExecutionTestBase
 
 		{
 			FConstSharedStruct SharedFragmentInstance = FConstSharedStruct::Make<FTestConstSharedFragment_Int>(TestIntValue);
-			Values.AddConstSharedFragment(SharedFragmentInstance);
+			Values.Add(SharedFragmentInstance);
 		}
 
 		const FTestConstSharedFragment_Int* ConstInstance = GetConstSharedFragmentPtr<FTestConstSharedFragment_Int>(Values);
@@ -98,7 +98,7 @@ struct FSharedFragmentValues_Contains : FExecutionTestBase
 
 		{
 			FSharedStruct SharedFragmentInstance = FSharedStruct::Make<FTestSharedFragment_Int>(TestIntValue);
-			Values.AddSharedFragment(SharedFragmentInstance);
+			Values.Add(SharedFragmentInstance);
 		}
 
 		AITEST_TRUE("Empty FMassArchetypeSharedFragmentValues should fail ContainsType tests"
@@ -108,7 +108,7 @@ struct FSharedFragmentValues_Contains : FExecutionTestBase
 
 		{
 			FSharedStruct SharedFragmentInstance = FSharedStruct::Make<FTestSharedFragment_Float>(TestFloatValue);
-			Values.AddSharedFragment(SharedFragmentInstance);
+			Values.Add(SharedFragmentInstance);
 		}
 		AITEST_TRUE("Empty FMassArchetypeSharedFragmentValues should fail ContainsType tests"
 			, Values.ContainsType(FTestSharedFragment_Float::StaticStruct()));
@@ -128,12 +128,12 @@ struct FSharedFragmentValues_Append : FExecutionTestBase
 		constexpr float TestFloatValue = 63.f;
 		
 		FMassArchetypeSharedFragmentValues ValuesNonConstInt;
-		ValuesNonConstInt.AddSharedFragment(FSharedStruct::Make<FTestSharedFragment_Int>(TestIntValue));
+		ValuesNonConstInt.Add(FSharedStruct::Make<FTestSharedFragment_Int>(TestIntValue));
 		FMassArchetypeSharedFragmentValues ValuesNonConstFloat;
-		ValuesNonConstFloat.AddSharedFragment(FSharedStruct::Make<FTestSharedFragment_Float>(TestFloatValue));
+		ValuesNonConstFloat.Add(FSharedStruct::Make<FTestSharedFragment_Float>(TestFloatValue));
 		FMassArchetypeSharedFragmentValues ValuestNonConstIntFloat;
-		ValuestNonConstIntFloat.AddSharedFragment(FSharedStruct::Make<FTestSharedFragment_Int>(TestIntValue));
-		ValuestNonConstIntFloat.AddSharedFragment(FSharedStruct::Make<FTestSharedFragment_Float>(TestFloatValue));
+		ValuestNonConstIntFloat.Add(FSharedStruct::Make<FTestSharedFragment_Int>(TestIntValue));
+		ValuestNonConstIntFloat.Add(FSharedStruct::Make<FTestSharedFragment_Float>(TestFloatValue));
 
 		// `Appending` int/float values to an new Values instance should result in the same result as `Adding` them
 		{
@@ -152,12 +152,12 @@ struct FSharedFragmentValues_Append : FExecutionTestBase
 		}
 
 		FMassArchetypeSharedFragmentValues ValuesConstInt;
-		ValuesConstInt.AddConstSharedFragment(FConstSharedStruct::Make<FTestConstSharedFragment_Int>(TestIntValue));
+		ValuesConstInt.Add(FConstSharedStruct::Make<FTestConstSharedFragment_Int>(TestIntValue));
 		FMassArchetypeSharedFragmentValues ValuesConstFloat;
-		ValuesConstFloat.AddConstSharedFragment(FConstSharedStruct::Make<FTestConstSharedFragment_Float>(TestFloatValue));
+		ValuesConstFloat.Add(FConstSharedStruct::Make<FTestConstSharedFragment_Float>(TestFloatValue));
 		FMassArchetypeSharedFragmentValues ValuestConstIntFloat;
-		ValuestConstIntFloat.AddConstSharedFragment(FConstSharedStruct::Make<FTestConstSharedFragment_Int>(TestIntValue));
-		ValuestConstIntFloat.AddConstSharedFragment(FConstSharedStruct::Make<FTestConstSharedFragment_Float>(TestFloatValue));
+		ValuestConstIntFloat.Add(FConstSharedStruct::Make<FTestConstSharedFragment_Int>(TestIntValue));
+		ValuestConstIntFloat.Add(FConstSharedStruct::Make<FTestConstSharedFragment_Float>(TestFloatValue));
 
 		{
 			FMassArchetypeSharedFragmentValues Values;
@@ -188,12 +188,12 @@ struct FSharedFragmentValues_Remove : FExecutionTestBase
 
 		{
 			FMassArchetypeSharedFragmentValues ValuesNonConstInt;
-			ValuesNonConstInt.AddSharedFragment(FSharedStruct::Make<FTestSharedFragment_Int>(TestIntValue));
+			ValuesNonConstInt.Add(FSharedStruct::Make<FTestSharedFragment_Int>(TestIntValue));
 			FMassArchetypeSharedFragmentValues ValuesNonConstFloat;
-			ValuesNonConstFloat.AddSharedFragment(FSharedStruct::Make<FTestSharedFragment_Float>(TestFloatValue));
+			ValuesNonConstFloat.Add(FSharedStruct::Make<FTestSharedFragment_Float>(TestFloatValue));
 			FMassArchetypeSharedFragmentValues ValuestNonConstIntFloat;
-			ValuestNonConstIntFloat.AddSharedFragment(FSharedStruct::Make<FTestSharedFragment_Int>(TestIntValue));
-			ValuestNonConstIntFloat.AddSharedFragment(FSharedStruct::Make<FTestSharedFragment_Float>(TestFloatValue));
+			ValuestNonConstIntFloat.Add(FSharedStruct::Make<FTestSharedFragment_Int>(TestIntValue));
+			ValuestNonConstIntFloat.Add(FSharedStruct::Make<FTestSharedFragment_Float>(TestFloatValue));
 
 			{
 				FMassArchetypeSharedFragmentValues Values = ValuestNonConstIntFloat;
@@ -212,12 +212,12 @@ struct FSharedFragmentValues_Remove : FExecutionTestBase
 		}
 		{
 			FMassArchetypeSharedFragmentValues ValuesConstInt;
-			ValuesConstInt.AddConstSharedFragment(FConstSharedStruct::Make<FTestConstSharedFragment_Int>(TestIntValue));
+			ValuesConstInt.Add(FConstSharedStruct::Make<FTestConstSharedFragment_Int>(TestIntValue));
 			FMassArchetypeSharedFragmentValues ValuesConstFloat;
-			ValuesConstFloat.AddConstSharedFragment(FConstSharedStruct::Make<FTestConstSharedFragment_Float>(TestFloatValue));
+			ValuesConstFloat.Add(FConstSharedStruct::Make<FTestConstSharedFragment_Float>(TestFloatValue));
 			FMassArchetypeSharedFragmentValues ValuestConstIntFloat;
-			ValuestConstIntFloat.AddConstSharedFragment(FConstSharedStruct::Make<FTestConstSharedFragment_Int>(TestIntValue));
-			ValuestConstIntFloat.AddConstSharedFragment(FConstSharedStruct::Make<FTestConstSharedFragment_Float>(TestFloatValue));
+			ValuestConstIntFloat.Add(FConstSharedStruct::Make<FTestConstSharedFragment_Int>(TestIntValue));
+			ValuestConstIntFloat.Add(FConstSharedStruct::Make<FTestConstSharedFragment_Float>(TestFloatValue));
 
 			{
 				FMassArchetypeSharedFragmentValues Values = ValuestConstIntFloat;
@@ -248,8 +248,8 @@ struct FSharedFragmentValues_Hash : FExecutionTestBase
 		constexpr float TestFloatValue = 63.f;
 
 		FMassArchetypeSharedFragmentValues ValuestNonConstIntFloat;
-		ValuestNonConstIntFloat.AddSharedFragment(FSharedStruct::Make<FTestSharedFragment_Int>(TestIntValue));
-		ValuestNonConstIntFloat.AddSharedFragment(FSharedStruct::Make<FTestSharedFragment_Float>(TestFloatValue));
+		ValuestNonConstIntFloat.Add(FSharedStruct::Make<FTestSharedFragment_Int>(TestIntValue));
+		ValuestNonConstIntFloat.Add(FSharedStruct::Make<FTestSharedFragment_Float>(TestFloatValue));
 
 		AITEST_SCOPED_CHECK("Expecting the containers to be sorted", 1);
 		const uint32 EmptyHash = ValuestNonConstIntFloat.CalculateHash();
@@ -398,11 +398,11 @@ struct FSharedFragmentBase : public FEntityTestBase
 
 		if constexpr (TIsDerivedFrom<TSharedStruct, FSharedStruct>::IsDerived)
 		{
-			SharedIntValues.AddSharedFragment(SharedFragmentInstance);
+			SharedIntValues.Add(SharedFragmentInstance);
 		}
 		else
 		{
-			SharedIntValues.AddConstSharedFragment(SharedFragmentInstance);
+			SharedIntValues.Add(SharedFragmentInstance);
 		}
 
 		EntityManager->BatchCreateEntities(FloatsArchetype, SharedIntValues, NumToCreate, OutEntityHandles);
@@ -418,11 +418,11 @@ struct FSharedFragmentBase : public FEntityTestBase
 
 		if constexpr (TIsDerivedFrom<TSharedStruct, FSharedStruct>::IsDerived)
 		{
-			SharedIntValues.AddSharedFragment(SharedFragmentInstance);
+			SharedIntValues.Add(SharedFragmentInstance);
 		}
 		else
 		{
-			SharedIntValues.AddConstSharedFragment(SharedFragmentInstance);
+			SharedIntValues.Add(SharedFragmentInstance);
 		}
 
 		OutEntityHandle = EntityManager->CreateEntity(FloatsArchetype, SharedIntValues);
@@ -692,7 +692,7 @@ struct FSharedFragment_BatchAddToEntity : FEntityTestBase
 
 		FMassArchetypeSharedFragmentValues SharedValues;
 		FConstSharedStruct ConstSharedFragment = FConstSharedStruct::Make<FTestConstSharedFragment_Int>(TestIntValue);
-		SharedValues.AddConstSharedFragment(ConstSharedFragment);
+		SharedValues.Add(ConstSharedFragment);
 		EntityManager->BatchAddSharedFragmentsForEntities(MakeArrayView(&EntityCollection, 1), SharedValues);
 
 		const FMassArchetypeHandle TargetArchetype = EntityManager->GetArchetypeForEntityUnsafe(EntitiesToMove[0]);
@@ -724,7 +724,7 @@ struct FSharedFragment_BatchSetAttempt : FEntityTestBase
 		TArray<FMassEntityHandle> CreatedEntityHandles;
 		FMassArchetypeSharedFragmentValues SharedIntValues;
 		FConstSharedStruct ConstSharedFragment = FConstSharedStruct::Make<FTestConstSharedFragment_Int>(TestIntValue);
-		SharedIntValues.AddConstSharedFragment(ConstSharedFragment);
+		SharedIntValues.Add(ConstSharedFragment);
 
 		TSharedRef<FMassEntityManager::FEntityCreationContext> CreationContext = EntityManager->BatchCreateEntities(FloatsArchetype, SharedIntValues, EntitiesToCreateNum, CreatedEntityHandles);
 		const FMassArchetypeHandle ResultingArchetype = CreationContext->GetEntityCollections()[0].GetArchetype();
@@ -740,7 +740,7 @@ struct FSharedFragment_BatchSetAttempt : FEntityTestBase
 		{
 			FMassArchetypeSharedFragmentValues DifferentSharedIntValues;
 			FConstSharedStruct OtherConstSharedFragment = FConstSharedStruct::Make<FTestConstSharedFragment_Int>(OtherTestIntValue);
-			DifferentSharedIntValues.AddConstSharedFragment(OtherConstSharedFragment);
+			DifferentSharedIntValues.Add(OtherConstSharedFragment);
 
 			AITEST_SCOPED_CHECK("Setting shared fragment values without archetype change is not supported", 1);
 			AITEST_SCOPED_CHECK("Trying to set shared fragment values, without adding new shared fragments", 1);
@@ -760,7 +760,7 @@ struct FSharedFragment_BatchAddToEmpty : FEntityTestBase
 
 		FMassArchetypeSharedFragmentValues SharedIntValues;
 		FSharedStruct SharedFragment = FSharedStruct::Make<FTestSharedFragment_Int>();
-		SharedIntValues.AddSharedFragment(SharedFragment);
+		SharedIntValues.Add(SharedFragment);
 
 		TArray<FMassEntityHandle> ReservedEntityHandles;
 		EntityManager->BatchReserveEntities(NumToReserve, ReservedEntityHandles);
@@ -792,7 +792,7 @@ struct FSharedFragmentValues_TypeEquivalency : FEntityTestBase
 		const FMassSharedFragmentBitSet IntSharedFragmentBitSet = FMassSharedFragmentBitSet::GetTypeBitSet<FTestSharedFragment_Int>();
 		FMassSharedFragmentBitSet IntFloatSharedFragmentBitSet = IntSharedFragmentBitSet;
 		IntFloatSharedFragmentBitSet.Add<FTestSharedFragment_Float>();
-		Values.AddSharedFragment(FSharedStruct::Make<FTestSharedFragment_Int>(TestIntValue));
+		Values.Add(FSharedStruct::Make<FTestSharedFragment_Int>(TestIntValue));
 		AITEST_TRUE("Single shared value type matches expected bitset", Values.HasExactSharedFragmentTypesMatch(IntSharedFragmentBitSet));
 		AITEST_FALSE("Single shared value type doesn't match two-type bitset", Values.HasExactSharedFragmentTypesMatch(IntFloatSharedFragmentBitSet));
 		AITEST_FALSE("Single shared value type doesn't match empty", Values.HasExactSharedFragmentTypesMatch(EmptySharedFragmentBitSet));
@@ -800,7 +800,7 @@ struct FSharedFragmentValues_TypeEquivalency : FEntityTestBase
 		const FMassConstSharedFragmentBitSet IntConstSharedFragmentBitSet = FMassConstSharedFragmentBitSet::GetTypeBitSet<FTestConstSharedFragment_Int>();
 		FMassConstSharedFragmentBitSet IntFloatConstSharedFragmentBitSet = IntConstSharedFragmentBitSet;
 		IntFloatConstSharedFragmentBitSet.Add<FTestConstSharedFragment_Float>();
-		Values.AddConstSharedFragment(FSharedStruct::Make<FTestConstSharedFragment_Int>(TestIntValue));
+		Values.Add(FSharedStruct::Make<FTestConstSharedFragment_Int>(TestIntValue));
 		AITEST_TRUE("Single const shared value type matches expected bitset", Values.HasExactConstSharedFragmentTypesMatch(IntConstSharedFragmentBitSet));
 		AITEST_FALSE("Single const shared value type doesn't match two-type bitset", Values.HasExactConstSharedFragmentTypesMatch(IntFloatConstSharedFragmentBitSet));
 		AITEST_FALSE("Single const shared value type doesn't match empty", Values.HasExactConstSharedFragmentTypesMatch(EmptyConstSharedFragmentBitSet));
diff --git a/Engine/Source/Runtime/CoreUObject/Public/StructUtils/InstancedStruct.h b/Engine/Source/Runtime/CoreUObject/Public/StructUtils/InstancedStruct.h
--- a/Engine/Source/Runtime/CoreUObject/Public/StructUtils/InstancedStruct.h
+++ b/Engine/Source/Runtime/CoreUObject/Public/StructUtils/InstancedStruct.h
@@ -24,7 +24,7 @@ class UUserDefinedStruct;
  *	TArray<FInstancedStruct> TestArray;
  */
 USTRUCT(BlueprintType, meta = (DisableSplitPin, HasNativeMake = "/Script/Engine.BlueprintInstancedStructLibrary.MakeInstancedStruct"))
-struct COREUOBJECT_API FInstancedStruct
+struct [[nodiscard]] COREUOBJECT_API FInstancedStruct
 {
 	GENERATED_BODY()
 
@@ -117,7 +117,7 @@ public:
 
 	/** Creates a new FInstancedStruct from templated struct type. */
 	template<typename T>
-	[[nodiscard]] static FInstancedStruct Make()
+	static FInstancedStruct Make()
 	{
 		UE::StructUtils::CheckStructType<T>();
 
@@ -128,7 +128,7 @@ public:
 
 	/** Creates a new FInstancedStruct from templated struct. */
 	template<typename T>
-	[[nodiscard]] static FInstancedStruct Make(const T& Struct)
+	static FInstancedStruct Make(const T& Struct)
 	{
 		UE::StructUtils::CheckStructType<T>();
 
@@ -139,7 +139,7 @@ public:
 
 	/** Creates a new FInstancedStruct from the templated type and forward all arguments to constructor. */
 	template<typename T, typename... TArgs>
-	[[nodiscard]] static FInstancedStruct Make(TArgs&&... InArgs)
+	static FInstancedStruct Make(TArgs&&... InArgs)
 	{
 		UE::StructUtils::CheckStructType<T>();
 
@@ -299,7 +299,7 @@ struct TStructOpsTypeTraits<FInstancedStruct> : public TStructOpsTypeTraitsBase2
  *	TArray<TInstancedStruct<FTestStructBase>> TestArray;
  */
 template<typename BaseStructT>
-struct TInstancedStruct
+struct [[nodiscard]] TInstancedStruct
 {
 public:
 	TInstancedStruct() = default;
@@ -372,7 +372,7 @@ public:
 
 	/** Creates a new TInstancedStruct from templated struct type. */
 	template<typename T = BaseStructT, typename = std::enable_if_t<std::is_base_of_v<BaseStructT, std::decay_t<T>>>>
-	[[nodiscard]] static TInstancedStruct Make()
+	static TInstancedStruct Make()
 	{
 		TInstancedStruct This;
 		This.InstancedStruct.InitializeAs(TBaseStructure<T>::Get(), nullptr);
@@ -381,7 +381,7 @@ public:
 
 	/** Creates a new TInstancedStruct from templated struct. */
 	template<typename T = BaseStructT, typename = std::enable_if_t<std::is_base_of_v<BaseStructT, std::decay_t<T>>>>
-	[[nodiscard]] static TInstancedStruct Make(const T& Struct)
+	static TInstancedStruct Make(const T& Struct)
 	{
 		TInstancedStruct This;
 		This.InstancedStruct.InitializeAs(TBaseStructure<T>::Get(), reinterpret_cast<const uint8*>(&Struct));
@@ -390,7 +390,7 @@ public:
 
 	/** Creates a new TInstancedStruct from the templated type and forward all arguments to constructor. */
 	template<typename T = BaseStructT, typename... TArgs, typename = std::enable_if_t<std::is_base_of_v<BaseStructT, std::decay_t<T>>>>
-	[[nodiscard]] static TInstancedStruct Make(TArgs&&... InArgs)
+	static TInstancedStruct Make(TArgs&&... InArgs)
 	{
 		TInstancedStruct This;
 		This.InstancedStruct.template InitializeAs<T>(Forward<TArgs>(InArgs)...);
diff --git a/Engine/Source/Runtime/CoreUObject/Public/StructUtils/SharedStruct.h b/Engine/Source/Runtime/CoreUObject/Public/StructUtils/SharedStruct.h
--- a/Engine/Source/Runtime/CoreUObject/Public/StructUtils/SharedStruct.h
+++ b/Engine/Source/Runtime/CoreUObject/Public/StructUtils/SharedStruct.h
@@ -152,7 +152,7 @@ private:
  * In either case, the shared struct memory /data is mutable.
  */
 USTRUCT()
-struct COREUOBJECT_API FSharedStruct
+struct [[nodiscard]] COREUOBJECT_API FSharedStruct
 {
 	GENERATED_BODY();
 
@@ -246,7 +246,7 @@ struct COREUOBJECT_API FSharedStruct
 
 	/** Creates a new FSharedStruct from a templated struct type. This will create a new instance of the shared struct memory. */
 	template<typename T>
-	[[nodiscard]] static FSharedStruct Make()
+	static FSharedStruct Make()
 	{
 		FSharedStruct SharedStruct;
 		SharedStruct.InitializeAs<T>();
@@ -255,7 +255,7 @@ struct COREUOBJECT_API FSharedStruct
 
 	/** Creates a new FSharedStruct from templated struct instance. This will create a new instance of the shared struct memory. */
 	template<typename T>
-	[[nodiscard]] static FSharedStruct Make(const T& Struct)
+	static FSharedStruct Make(const T& Struct)
 	{
 		FSharedStruct SharedStruct;
 		SharedStruct.InitializeAs(Struct);
@@ -263,7 +263,7 @@ struct COREUOBJECT_API FSharedStruct
 	}
 
 	/** Creates a new FSharedStruct from struct type and optional instance memory. This will create a new instance of the shared struct memory. */
-	[[nodiscard]] static FSharedStruct Make(const UScriptStruct* InScriptStruct, const uint8* InStructMemory = nullptr)
+	static FSharedStruct Make(const UScriptStruct* InScriptStruct, const uint8* InStructMemory = nullptr)
 	{
 		FSharedStruct SharedStruct;
 		SharedStruct.InitializeAs(InScriptStruct, InStructMemory);
@@ -272,7 +272,7 @@ struct COREUOBJECT_API FSharedStruct
 
 	/** Creates a new FSharedStruct from the templated type and forward all arguments to constructor. This will create a new instance of the shared struct memory. */
 	template<typename T, typename... TArgs>
-	[[nodiscard]] static FSharedStruct Make(TArgs&&... InArgs)
+	static FSharedStruct Make(TArgs&&... InArgs)
 	{
 		FSharedStruct SharedStruct;
 		SharedStruct.InitializeAs<T>(Forward<TArgs>(InArgs)...);
@@ -340,7 +340,7 @@ protected:
  *	TArray<TSharedStruct<FTestStructBase>> TestArray;
  */
 template<typename BaseStructT>
-struct TSharedStruct
+struct [[nodiscard]] TSharedStruct
 {
 	explicit TSharedStruct() = default;
 
@@ -431,21 +431,21 @@ struct TSharedStruct
 
 	/** Creates a new TSharedStruct. This will create a new instance of the shared struct memory. */
 	template<typename T = BaseStructT, typename = std::enable_if_t<std::is_base_of_v<BaseStructT, std::decay_t<T>>>>
-	[[nodiscard]] static TSharedStruct Make()
+	static TSharedStruct Make()
 	{
 		return FSharedStruct::Make<T>();
 	}
 
 	/** Creates a new TSharedStruct from templated struct instance. This will create a new instance of the shared struct memory. */
 	template<typename T = BaseStructT, typename = std::enable_if_t<std::is_base_of_v<BaseStructT, std::decay_t<T>>>>
-	[[nodiscard]] static TSharedStruct Make(const T& Struct)
+	static TSharedStruct Make(const T& Struct)
 	{
 		return FSharedStruct::Make<T>(Struct);
 	}
 
 	/** Creates a new TSharedStruct from the templated type and forward all arguments to constructor. This will create a new instance of the shared struct memory. */
 	template< typename... TArgs>
-	[[nodiscard]] static TSharedStruct Make(TArgs&&... InArgs)
+	static TSharedStruct Make(TArgs&&... InArgs)
 	{
 		return FSharedStruct::Make<BaseStructT>(Forward<TArgs>(InArgs)...);
 	}
@@ -517,7 +517,7 @@ struct TStructOpsTypeTraits<FSharedStruct> : public TStructOpsTypeTraitsBase2<FS
  * See FSharedStruct for more information.
  */
 USTRUCT()
-struct COREUOBJECT_API FConstSharedStruct
+struct [[nodiscard]] COREUOBJECT_API FConstSharedStruct
 {
 	GENERATED_BODY();
 
@@ -622,7 +622,7 @@ struct COREUOBJECT_API FConstSharedStruct
 
 	/** Creates a new FSharedStruct from a templated struct type. This will create a new instance of the shared struct memory. */
 	template<typename T>
-	[[nodiscard]] static FConstSharedStruct Make()
+	static FConstSharedStruct Make()
 	{
 		FConstSharedStruct SharedStruct;
 		SharedStruct.InitializeAs<T>();
@@ -631,7 +631,7 @@ struct COREUOBJECT_API FConstSharedStruct
 
 	/** Creates a new FSharedStruct from templated struct instance. This will create a new instance of the shared struct memory. */
 	template<typename T>
-	[[nodiscard]] static FConstSharedStruct Make(const T& Struct)
+	static FConstSharedStruct Make(const T& Struct)
 	{
 		FConstSharedStruct SharedStruct;
 		SharedStruct.InitializeAs<T>(Struct);
@@ -639,7 +639,7 @@ struct COREUOBJECT_API FConstSharedStruct
 	}
 
 	/** Creates a new FSharedStruct from a struct type and optional data. This will create a new instance of the shared struct memory. */
-	[[nodiscard]] static FConstSharedStruct Make(const UScriptStruct* InScriptStruct, const uint8* InStructMemory = nullptr)
+	static FConstSharedStruct Make(const UScriptStruct* InScriptStruct, const uint8* InStructMemory = nullptr)
 	{
 		FConstSharedStruct SharedStruct;
 		SharedStruct.InitializeAs(InScriptStruct, InStructMemory);
@@ -648,7 +648,7 @@ struct COREUOBJECT_API FConstSharedStruct
 
 	/** Creates a new FSharedStruct from the templated type and forward all arguments to constructor. This will create a new instance of the shared struct memory. */
 	template<typename T, typename... TArgs>
-	[[nodiscard]] static FConstSharedStruct Make(TArgs&&... InArgs)
+	static FConstSharedStruct Make(TArgs&&... InArgs)
 	{
 		FConstSharedStruct SharedStruct;
 		SharedStruct.InitializeAs<T>(Forward<TArgs>(InArgs)...);
@@ -717,7 +717,7 @@ protected:
  *	TArray<TConstSharedStruct<FTestStructBase>> TestArray;
  */
 template<typename BaseStructT>
-struct TConstSharedStruct
+struct [[nodiscard]] TConstSharedStruct
 {
 	explicit TConstSharedStruct() = default;
 
@@ -802,7 +802,7 @@ struct TConstSharedStruct
 
 	/** Creates a new TSharedStruct from a templated struct type. This will create a new instance of the shared struct memory. */
 	template<typename T = BaseStructT, typename = std::enable_if_t<std::is_base_of_v<BaseStructT, std::decay_t<T>>>>
-	[[nodiscard]] static TConstSharedStruct Make()
+	static TConstSharedStruct Make()
 	{
 		FConstSharedStruct SharedStruct;
 		SharedStruct.InitializeAs<BaseStructT>();
@@ -811,7 +811,7 @@ struct TConstSharedStruct
 
 	/** Creates a new TSharedStruct from templated struct instance. This will create a new instance of the shared struct memory. */
 	template<typename T = BaseStructT, typename = std::enable_if_t<std::is_base_of_v<BaseStructT, std::decay_t<T>>>>
-	[[nodiscard]] static TConstSharedStruct Make(const T& Struct)
+	static TConstSharedStruct Make(const T& Struct)
 	{
 		FConstSharedStruct SharedStruct;
 		SharedStruct.InitializeAs<BaseStructT>(Struct);
@@ -820,7 +820,7 @@ struct TConstSharedStruct
 
 	/** Creates a new TSharedStruct from a struct type and optional data. This will create a new instance of the shared struct memory. */
 	template<typename T = BaseStructT, typename = std::enable_if_t<std::is_base_of_v<BaseStructT, std::decay_t<T>>>>
-	[[nodiscard]] static TConstSharedStruct Make(const uint8* InStructMemory = nullptr)
+	static TConstSharedStruct Make(const uint8* InStructMemory = nullptr)
 	{
 		TConstSharedStruct SharedStruct;
 		SharedStruct.template Initialize<T>(InStructMemory);
@@ -829,7 +829,7 @@ struct TConstSharedStruct
 
 	/** Creates a new TSharedStruct from the templated type and forward all arguments to constructor. This will create a new instance of the shared struct memory. */
 	template<typename T = BaseStructT, typename... TArgs, typename = std::enable_if_t<std::is_base_of_v<BaseStructT, std::decay_t<T>>>>
-	[[nodiscard]] static TConstSharedStruct Make(TArgs&&... InArgs)
+	static TConstSharedStruct Make(TArgs&&... InArgs)
 	{
 		TConstSharedStruct SharedStruct;
 		SharedStruct.template Initialize<T>(Forward<TArgs>(InArgs)...);
diff --git a/Engine/Source/Runtime/CoreUObject/Public/StructUtils/StructView.h b/Engine/Source/Runtime/CoreUObject/Public/StructUtils/StructView.h
--- a/Engine/Source/Runtime/CoreUObject/Public/StructUtils/StructView.h
+++ b/Engine/Source/Runtime/CoreUObject/Public/StructUtils/StructView.h
@@ -19,7 +19,7 @@
  * data can be mutated. Use FConstStructView to prevent mutation of the actual struct data.
  * See FConstStructView for examples.
  */
-struct FStructView
+struct [[nodiscard]] FStructView
 {
 public:
 
@@ -40,7 +40,7 @@ public:
 
 	/** Creates a new FStructView from the templated struct. Note its not safe to make InStruct const ref as the original object may have been declared const */
 	template<typename T>
-	[[nodiscard]] static FStructView Make(T& InStruct)
+	static FStructView Make(T& InStruct)
 	{
 		UE::StructUtils::CheckStructType<T>();
 		return FStructView(TBaseStructure<T>::Get(), reinterpret_cast<uint8*>(&InStruct));
@@ -113,7 +113,7 @@ protected:
  *	TArray<TStructView<FTestStructBase>> TestArray;
  */
 template<typename BaseStructT>
-struct TStructView : FStructView
+struct [[nodiscard]] TStructView : FStructView
 {
 public:
 
@@ -199,7 +199,7 @@ public:
  * e.g. FStructView A; const FStructView B; A = B; // compiles as the struct B is pointing to can't be made to point at something else but A isn't const.
  * e.g. const FStructView A; FStructView B; A = B; // doesn't compile as attempting to make const view point at something else
  */
-struct FConstStructView
+struct [[nodiscard]] FConstStructView
 {
 public:
 
@@ -229,7 +229,7 @@ public:
 
 	/** Creates a new FConstStructView from the templated struct */
 	template<typename T>
-	[[nodiscard]] static FConstStructView Make(const T& Struct)
+	static FConstStructView Make(const T& Struct)
 	{
 		UE::StructUtils::CheckStructType<T>();
 		return FConstStructView(TBaseStructure<T>::Get(), reinterpret_cast<const uint8*>(&Struct));
@@ -308,7 +308,7 @@ protected:
  *	TArray<TConstStructView<FTestStructBase>> TestArray;
  */
 template<typename BaseStructT>
-struct TConstStructView : FConstStructView
+struct [[nodiscard]] TConstStructView : FConstStructView
 {
 	explicit TConstStructView() = default;
 
diff --git a/Engine/Source/Runtime/MassEntity/Private/MassEntityBuilder.cpp b/Engine/Source/Runtime/MassEntity/Private/MassEntityBuilder.cpp
--- a/Engine/Source/Runtime/MassEntity/Private/MassEntityBuilder.cpp
+++ b/Engine/Source/Runtime/MassEntity/Private/MassEntityBuilder.cpp
@@ -459,13 +459,13 @@ void FEntityBuilder::CacheSharedFragmentValue()
 		{
 			check(SharedFragmentInstance.IsValid());
 			const FSharedStruct& SharedStruct = EntityManager->GetOrCreateSharedFragment(*SharedFragmentInstance.GetScriptStruct(), SharedFragmentInstance.GetMemory());
-			CachedSharedFragmentValues.AddSharedFragment(SharedStruct);
+			CachedSharedFragmentValues.Add(SharedStruct);
 		}
 		for (FInstancedStruct& ConstSharedFragmentInstance : ConstSharedFragments)
 		{
 			check(ConstSharedFragmentInstance.IsValid());
 			const FConstSharedStruct& ConstSharedStruct = EntityManager->GetOrCreateConstSharedFragment(*ConstSharedFragmentInstance.GetScriptStruct(), ConstSharedFragmentInstance.GetMemory());
-			CachedSharedFragmentValues.AddConstSharedFragment(ConstSharedStruct);
+			CachedSharedFragmentValues.Add(ConstSharedStruct);
 		}
 
 		CachedSharedFragmentValues.Sort();
diff --git a/Engine/Source/Runtime/MassEntity/Private/MassEntityManager.cpp b/Engine/Source/Runtime/MassEntity/Private/MassEntityManager.cpp
--- a/Engine/Source/Runtime/MassEntity/Private/MassEntityManager.cpp
+++ b/Engine/Source/Runtime/MassEntity/Private/MassEntityManager.cpp
@@ -1492,7 +1492,7 @@ bool FMassEntityManager::AddConstSharedFragmentToEntity(const FMassEntityHandle
 	const FMassArchetypeSharedFragmentValues& OldSharedFragmentValues = CurrentArchetype->GetSharedFragmentValues(Entity.Index);
 	check(!OldSharedFragmentValues.ContainsType(StructType));
 	FMassArchetypeSharedFragmentValues NewSharedFragmentValues(OldSharedFragmentValues);
-	NewSharedFragmentValues.AddConstSharedFragment(InConstSharedFragment);
+	NewSharedFragmentValues.Add(InConstSharedFragment);
 	NewSharedFragmentValues.Sort();
 
 	CurrentArchetype->MoveEntityToAnotherArchetype(Entity, *NewArchetype, &NewSharedFragmentValues);
@@ -1585,7 +1585,7 @@ bool FMassEntityManager::AddSharedFragmentToEntity(const FMassEntityHandle Entit
 	const FMassArchetypeSharedFragmentValues& OldSharedFragmentValues = CurrentArchetype->GetSharedFragmentValues(Entity.Index);
 	check(!OldSharedFragmentValues.ContainsType(StructType));
 	FMassArchetypeSharedFragmentValues NewSharedFragmentValues(OldSharedFragmentValues);
-	NewSharedFragmentValues.AddSharedFragment(InSharedFragment);
+	NewSharedFragmentValues.Add(InSharedFragment);
 	NewSharedFragmentValues.Sort();
 
 	CurrentArchetype->MoveEntityToAnotherArchetype(Entity, *NewArchetype, &NewSharedFragmentValues);
diff --git a/Engine/Source/Runtime/MassEntity/Private/MassEntityTypes.cpp b/Engine/Source/Runtime/MassEntity/Private/MassEntityTypes.cpp
--- a/Engine/Source/Runtime/MassEntity/Private/MassEntityTypes.cpp
+++ b/Engine/Source/Runtime/MassEntity/Private/MassEntityTypes.cpp
@@ -95,15 +95,17 @@ void FMassArchetypeCompositionDescriptor::DebugOutputDescription(FOutputDevice&
 //-----------------------------------------------------------------------------
 // FMassArchetypeSharedFragmentValues
 //-----------------------------------------------------------------------------
-FConstSharedStruct FMassArchetypeSharedFragmentValues::AddConstSharedFragment(const FConstSharedStruct& Fragment)
+FConstSharedStruct FMassArchetypeSharedFragmentValues::Add_GetRef(const FConstSharedStruct& Fragment)
 {
 	check(Fragment.IsValid());
 	const UScriptStruct* StructType = Fragment.GetScriptStruct();
-	if (!ensureMsgf(ContainsType(StructType) == false, TEXT("Shared Fragment of type %s already added to FMassArchetypeSharedFragmentValues"), *GetNameSafe(StructType)))
+	if (ContainsType(StructType))
 	{
-		FConstSharedStruct ExistingSharedStruct = GetConstSharedFragmentStruct(StructType);
-		ensureMsgf(ExistingSharedStruct.IsValid(), TEXT("Shared Fragment of type %s already added as NON-CONST shared struct"), *GetNameSafe(StructType));
-		return ExistingSharedStruct;
+		FConstSharedStruct ExistingConstSharedStruct = GetConstSharedFragmentStruct(StructType);
+		ensureMsgf(false, TEXT("Shared Fragment of type %s already added to FMassArchetypeSharedFragmentValues%s")
+			, *GetNameSafe(StructType)
+			, ExistingConstSharedStruct.IsValid() ? TEXT("") : TEXT(" as NON-CONST shared struct"));
+		return ExistingConstSharedStruct;
 	}
 
 	check(StructType);
@@ -113,14 +115,16 @@ FConstSharedStruct FMassArchetypeSharedFragmentValues::AddConstSharedFragment(co
 	return StructInstance;
 }
 
-FSharedStruct FMassArchetypeSharedFragmentValues::AddSharedFragment(const FSharedStruct& Fragment)
+FSharedStruct FMassArchetypeSharedFragmentValues::Add_GetRef(const FSharedStruct& Fragment)
 {
 	check(Fragment.IsValid());
 	const UScriptStruct* StructType = Fragment.GetScriptStruct();
-	if (!ensureMsgf(ContainsType(StructType) == false, TEXT("Shared Fragment of type %s already added to FMassArchetypeSharedFragmentValues"), *GetNameSafe(StructType)))
+	if (ContainsType(StructType))
 	{
 		FSharedStruct ExistingSharedStruct = GetSharedFragmentStruct(StructType);
-		ensureMsgf(ExistingSharedStruct.IsValid(), TEXT("Shared Fragment of type %s already added as CONST shared struct"), *GetNameSafe(StructType));
+		ensureMsgf(false, TEXT("Shared Fragment of type %s already added to FMassArchetypeSharedFragmentValues%s")
+			, *GetNameSafe(StructType)
+			, ExistingSharedStruct.IsValid() ? TEXT("") : TEXT(" as CONST shared struct"));
 		return ExistingSharedStruct;
 	}
 
diff --git a/Engine/Source/Runtime/MassEntity/Public/MassEntityTypes.h b/Engine/Source/Runtime/MassEntity/Public/MassEntityTypes.h
--- a/Engine/Source/Runtime/MassEntity/Public/MassEntityTypes.h
+++ b/Engine/Source/Runtime/MassEntity/Public/MassEntityTypes.h
@@ -308,19 +308,51 @@ struct MASSENTITY_API FMassArchetypeSharedFragmentValues
 		}
 	}
 
+	/**
+	 * Adds Fragment to the collection.
+	 * Method will ensure if a fragment of the given FMassConstSharedFragment subclass has already been added.
+	 */
+	void Add(const FConstSharedStruct& Fragment)
+	{
+		(void)Add_GetRef(Fragment);
+	}
+
 	/** 
-	 * Adds Fragment to the collection. If a fragment of the given FMassSharedFragment subclass has already added then 
-	 * the function will return the previously added instance. In that case the function will also assert if the given type 
-	 * has been added as a CONST shared fragment and if not it will return an empty FConstSharedStruct
+	 * Adds Fragment to the collection.
+	 * Method will ensure if a fragment of the given FMassConstSharedFragment subclass has already been added.
+	 * In that case the method will return the previously added instance if the given type has been added
+	 * as a CONST shared fragment and if not it will return an empty FConstSharedStruct.
 	 */
-	FConstSharedStruct AddConstSharedFragment(const FConstSharedStruct& Fragment);
+	FConstSharedStruct Add_GetRef(const FConstSharedStruct& Fragment);
+
+	UE_DEPRECATED(5.6, "Use Add or Add_GetRef instead depending on whether you need the return value.")
+		FConstSharedStruct AddConstSharedFragment(const FConstSharedStruct& Fragment)
+	{
+		return Add_GetRef(Fragment);
+	}
 
 	/**
-	 * Adds Fragment to the collection. If a fragment of the given FMassSharedFragment subclass has already added then
-	 * the function will return the previously added instance. In that case the function will also assert if the given type
-	 * has been added as a NON-CONST shared fragment and if not it will return an empty FSharedStruct
+	 * Adds Fragment to the collection.
+	 * Method will ensure if a fragment of the given FMassSharedFragment subclass has already been added.
 	 */
-	FSharedStruct AddSharedFragment(const FSharedStruct& Fragment);
+	void Add(const FSharedStruct& Fragment)
+	{
+		(void)Add_GetRef(Fragment);
+	}
+
+	/** 
+	 * Adds Fragment to the collection.
+	 * Method will ensure if a fragment of the given FMassSharedFragment subclass has already been added.
+	 * In that case the method will return the previously added instance if the given type has been added
+	 * as a NON-CONST shared fragment and if not it will return an empty FSharedStruct.
+	 */
+	FSharedStruct Add_GetRef(const FSharedStruct& Fragment);
+
+	UE_DEPRECATED(5.6, "Use Add or Add_GetRef instead depending on whether you need the return value.")
+	FSharedStruct AddSharedFragment(const FSharedStruct& Fragment)
+	{
+		return Add_GetRef(Fragment);
+	}
 
 	/**
 	 * Finds instances of fragment types given by Fragments and replaces their values with contents of respective
