Fixed CIS issue 841893: static analysis

diff --git a/Engine/Source/Developer/MassEntityTestSuite/Private/MassSharedFragmentsTest.cpp b/Engine/Source/Developer/MassEntityTestSuite/Private/MassSharedFragmentsTest.cpp
--- a/Engine/Source/Developer/MassEntityTestSuite/Private/MassSharedFragmentsTest.cpp
+++ b/Engine/Source/Developer/MassEntityTestSuite/Private/MassSharedFragmentsTest.cpp
@@ -395,15 +395,7 @@ struct FSharedFragmentBase : public FEntityTestBase
 		FMassArchetypeSharedFragmentValues SharedIntValues;
 
 		TSharedStruct SharedFragmentInstance = TSharedStruct::Make(FragmentInstance);
-
-		if constexpr (TIsDerivedFrom<TSharedStruct, FSharedStruct>::IsDerived)
-		{
-			SharedIntValues.Add(SharedFragmentInstance);
-		}
-		else
-		{
-			SharedIntValues.Add(SharedFragmentInstance);
-		}
+		SharedIntValues.Add(SharedFragmentInstance);
 
 		EntityManager->BatchCreateEntities(FloatsArchetype, SharedIntValues, NumToCreate, OutEntityHandles);
 	}
@@ -415,15 +407,7 @@ struct FSharedFragmentBase : public FEntityTestBase
 		FMassArchetypeSharedFragmentValues SharedIntValues;
 
 		TSharedStruct SharedFragmentInstance = TSharedStruct::Make(FragmentInstance);
-
-		if constexpr (TIsDerivedFrom<TSharedStruct, FSharedStruct>::IsDerived)
-		{
-			SharedIntValues.Add(SharedFragmentInstance);
-		}
-		else
-		{
-			SharedIntValues.Add(SharedFragmentInstance);
-		}
+		SharedIntValues.Add(SharedFragmentInstance);
 
 		OutEntityHandle = EntityManager->CreateEntity(FloatsArchetype, SharedIntValues);
 	}
@@ -800,7 +784,7 @@ struct FSharedFragmentValues_TypeEquivalency : FEntityTestBase
 		const FMassConstSharedFragmentBitSet IntConstSharedFragmentBitSet = FMassConstSharedFragmentBitSet::GetTypeBitSet<FTestConstSharedFragment_Int>();
 		FMassConstSharedFragmentBitSet IntFloatConstSharedFragmentBitSet = IntConstSharedFragmentBitSet;
 		IntFloatConstSharedFragmentBitSet.Add<FTestConstSharedFragment_Float>();
-		Values.Add(FSharedStruct::Make<FTestConstSharedFragment_Int>(TestIntValue));
+		Values.Add(FConstSharedStruct::Make<FTestConstSharedFragment_Int>(TestIntValue));
 		AITEST_TRUE("Single const shared value type matches expected bitset", Values.HasExactConstSharedFragmentTypesMatch(IntConstSharedFragmentBitSet));
 		AITEST_FALSE("Single const shared value type doesn't match two-type bitset", Values.HasExactConstSharedFragmentTypesMatch(IntFloatConstSharedFragmentBitSet));
 		AITEST_FALSE("Single const shared value type doesn't match empty", Values.HasExactConstSharedFragmentTypesMatch(EmptyConstSharedFragmentBitSet));
