[Mass] CIS fix, Issue 866655: PVS Studio warnings in MassObserverTest.cpp

#rnx

diff --git a/Engine/Source/Developer/MassEntityTestSuite/Private/MassObserverTest.cpp b/Engine/Source/Developer/MassEntityTestSuite/Private/MassObserverTest.cpp
--- a/Engine/Source/Developer/MassEntityTestSuite/Private/MassObserverTest.cpp
+++ b/Engine/Source/Developer/MassEntityTestSuite/Private/MassObserverTest.cpp
@@ -747,6 +747,8 @@ IMPLEMENT_AI_INSTANT_TEST(FModificationsWhileCreationContextActive, "System.Mass
 
 struct FCreationOperationOrder : FFragmentTestBase
 {
+	int32 Counter = 0;
+
 	FCreationOperationOrder()
 	{
 		OperationObserved = EMassObservedOperation::Add;
@@ -758,13 +760,12 @@ struct FCreationOperationOrder : FFragmentTestBase
 	{
 		TArray<FMassEntityHandle> PreExistingEntities;
 		TArray<FMassEntityHandle> NewEntities;
-		int32 Counter = 0;
 
-		ObserverProcessor->ForEachEntityChunkExecutionFunction = [&Counter](FMassExecutionContext& Context)
+		ObserverProcessor->ForEachEntityChunkExecutionFunction = [Counter = &this->Counter](FMassExecutionContext& Context)
 			{
 				for (auto EntityId : Context.CreateEntityIterator())
 				{
-					Context.GetMutableFragmentView<FFragmentStruct>()[EntityId].Value = static_cast<float>(++Counter);
+					Context.GetMutableFragmentView<FFragmentStruct>()[EntityId].Value = static_cast<float>(++(*Counter));
 				}
 			};
 
