[Mass]
Added MassQueryExecutor to unify the functions of MassQuery and MassProcessor and create a foundation to remove dependancy on UObject for Mass processing.
Simplified Mass data access with templated UE::Mass:FAccessors collection.

diff --git a/Engine/Source/Developer/MassEntityTestSuite/Private/MassQueryExecutorTest.cpp b/Engine/Source/Developer/MassEntityTestSuite/Private/MassQueryExecutorTest.cpp
new file mode 100644
--- /dev/null
+++ b/Engine/Source/Developer/MassEntityTestSuite/Private/MassQueryExecutorTest.cpp
@@ -0,0 +1,926 @@
+// Copyright Epic Games, Inc. All Rights Reserved.
+
+#include "AITestsCommon.h"
+#include "MassEntityManager.h"
+#include "MassProcessingTypes.h"
+#include "MassEntityTestTypes.h"
+#include "MassEntityUtils.h"
+#include "MassExecutor.h"
+#include "MassExecutionContext.h"
+#include "MassProcessingContext.h"
+#include "MassQueryExecutor.h"
+#include "HAL/PlatformTime.h"
+
+
+#define LOCTEXT_NAMESPACE "MassTest"
+
+DEFINE_LOG_CATEGORY_STATIC(LogMassPerfTest, Log, All);
+
+//UE_DISABLE_OPTIMIZATION_SHIP
+
+void UMassTestProcessorAutoExecuteQueryComparison::ConfigureQueries(const TSharedRef<FMassEntityManager>& EntityManager)
+{
+	EntityQuery.AddRequirement<FTestFragment_Float>(EMassFragmentAccess::ReadOnly);
+	EntityQuery.AddRequirement<FTestFragment_Int>(EMassFragmentAccess::ReadWrite);
+	EntityQuery.AddRequirement<FTestFragment_Bool>(EMassFragmentAccess::ReadOnly);
+	EntityQuery.AddRequirement<FTestFragment_Large>(EMassFragmentAccess::ReadOnly);
+	EntityQuery.AddRequirement<FTestFragment_Array>(EMassFragmentAccess::ReadOnly);
+}
+
+void UMassTestProcessorAutoExecuteQueryComparison::Execute(FMassEntityManager& EntityManager, FMassExecutionContext& Context)
+{
+	EntityQuery.ForEachEntityChunk(Context, [](FMassExecutionContext& Context)
+	{
+		const TConstArrayView<FTestFragment_Float> TestFloat = Context.GetFragmentView<FTestFragment_Float>();
+		const TArrayView<FTestFragment_Int> TestInt = Context.GetMutableFragmentView<FTestFragment_Int>();
+		const TConstArrayView<FTestFragment_Bool> TestBool = Context.GetFragmentView<FTestFragment_Bool>();
+		const TConstArrayView<FTestFragment_Large> TestLarge = Context.GetFragmentView<FTestFragment_Large>();
+		const TConstArrayView<FTestFragment_Array> TestArray = Context.GetFragmentView<FTestFragment_Array>();
+
+		for (uint32 EntityIndex : Context.CreateEntityIterator())
+		{
+			int32& TestIntVal = TestInt[EntityIndex].Value;
+			if (TestFloat[EntityIndex].Value > 0.0f)
+			{
+				++TestIntVal;
+			}
+			else
+			{
+				--TestIntVal;
+			}
+
+			if (TestBool[EntityIndex].bValue)
+			{
+				++TestIntVal;
+			}
+			else
+			{
+				--TestIntVal;
+			}
+
+			if (TestLarge[EntityIndex].Value[0] > 0)
+			{
+				++TestIntVal;
+			}
+			else
+			{
+				--TestIntVal;
+			}
+
+			if (TestArray[EntityIndex].Value.Num() > 0)
+			{
+				++TestIntVal;
+			}
+			else
+			{
+				--TestIntVal;
+			}
+		}
+	});
+}
+
+void UMassTestProcessorAutoExecuteQueryComparison_Parallel::ConfigureQueries(const TSharedRef<FMassEntityManager>& EntityManager)
+{
+	EntityQuery.AddRequirement<FTestFragment_Float>(EMassFragmentAccess::ReadOnly);
+	EntityQuery.AddRequirement<FTestFragment_Int>(EMassFragmentAccess::ReadWrite);
+	EntityQuery.AddRequirement<FTestFragment_Bool>(EMassFragmentAccess::ReadOnly);
+	EntityQuery.AddRequirement<FTestFragment_Large>(EMassFragmentAccess::ReadOnly);
+	EntityQuery.AddRequirement<FTestFragment_Array>(EMassFragmentAccess::ReadOnly);
+}
+
+void UMassTestProcessorAutoExecuteQueryComparison_Parallel::Execute(FMassEntityManager& EntityManager, FMassExecutionContext& Context)
+{
+	EntityQuery.ParallelForEachEntityChunk(Context, [](FMassExecutionContext& Context)
+		{
+			const TConstArrayView<FTestFragment_Float> TestFloat = Context.GetFragmentView<FTestFragment_Float>();
+			const TArrayView<FTestFragment_Int> TestInt = Context.GetMutableFragmentView<FTestFragment_Int>();
+			const TConstArrayView<FTestFragment_Bool> TestBool = Context.GetFragmentView<FTestFragment_Bool>();
+			const TConstArrayView<FTestFragment_Large> TestLarge = Context.GetFragmentView<FTestFragment_Large>();
+			const TConstArrayView<FTestFragment_Array> TestArray = Context.GetFragmentView<FTestFragment_Array>();
+
+			for (uint32 EntityIndex : Context.CreateEntityIterator())
+			{
+				int32& TestIntVal = TestInt[EntityIndex].Value;
+				if (TestFloat[EntityIndex].Value > 0.0f)
+				{
+					++TestIntVal;
+				}
+				else
+				{
+					--TestIntVal;
+				}
+
+				if (TestBool[EntityIndex].bValue)
+				{
+					++TestIntVal;
+				}
+				else
+				{
+					--TestIntVal;
+				}
+
+				if (TestLarge[EntityIndex].Value[0] > 0)
+				{
+					++TestIntVal;
+				}
+				else
+				{
+					--TestIntVal;
+				}
+
+				if (TestArray[EntityIndex].Value.Num() > 0)
+				{
+					++TestIntVal;
+				}
+				else
+				{
+					--TestIntVal;
+				}
+			}
+		});
+}
+
+namespace FMassQueryExecutorTest
+{
+
+	struct FTestQueryExecutor_LoadTestSetup : public UE::Mass::FQueryExecutor
+	{
+		UE::Mass::FQueryDefinition<
+			UE::Mass::FMutableFragmentAccess<FTestFragment_Int>,
+			UE::Mass::FMutableFragmentAccess<FTestFragment_Float>,
+			UE::Mass::FMutableFragmentAccess<FTestFragment_Bool>,
+			UE::Mass::FMutableFragmentAccess<FTestFragment_Large>,
+			UE::Mass::FMutableFragmentAccess<FTestFragment_Array>,
+			UE::Mass::FMutableSubsystemAccess<UMassTestWorldSubsystem>
+		> Accessors{ *this };
+
+		virtual void Execute(FMassExecutionContext& Context)
+		{
+			ParallelForEachEntityChunk(Context, Accessors, [](FMassExecutionContext& Context, auto& Data)
+			{
+				for (uint32 EntityIndex : Context.CreateEntityIterator())
+				{
+					Data.template Get<FTestFragment_Int>()[EntityIndex].Value = 0;
+					Data.template Get<FTestFragment_Float>()[EntityIndex].Value = FMath::FRandRange(-1.0f, 1.0f);
+					Data.template Get<FTestFragment_Bool>()[EntityIndex].bValue = FMath::RandBool();
+					Data.template Get<FTestFragment_Large>()[EntityIndex].Value[0] = FMath::RandBool() ? 1 : 0;
+					Data.template Get<FTestFragment_Array>()[EntityIndex].Value.SetNum(FMath::RandBool() ? 1 : 0);
+					Data.template Get<UMassTestWorldSubsystem>()->Write(0);
+				}
+			});
+		}
+	};
+
+	struct FTestQueryExecutor_LoadTestReset : public UE::Mass::FQueryExecutor
+	{
+		UE::Mass::FQueryDefinition<
+			UE::Mass::FMutableFragmentAccess<FTestFragment_Int>
+		> Accessors{ *this };
+
+		virtual void Execute(FMassExecutionContext& Context)
+		{
+			ParallelForEachEntityChunk(Context, Accessors, [](FMassExecutionContext& Context, auto& Data)
+			{
+				for (uint32 EntityIndex : Context.CreateEntityIterator())
+				{
+					Data.template Get<FTestFragment_Int>()[EntityIndex].Value = 0;
+				}
+			});
+		}
+	};
+
+	struct FTestQueryExecutor_LoadTestValidate : public UE::Mass::FQueryExecutor
+	{
+		UE::Mass::FQueryDefinition<
+			UE::Mass::FConstFragmentAccess<FTestFragment_Int>
+		> Accessors{ *this };
+
+		std::atomic<int64> Sum = 0;
+
+		virtual void Execute(FMassExecutionContext& Context)
+		{
+			ParallelForEachEntityChunk(Context, Accessors, [&Sum = Sum](FMassExecutionContext& Context, auto& Data)
+			{
+				int64 ChunkSum = 0;
+				for (uint32 EntityIndex : Context.CreateEntityIterator())
+				{
+					ChunkSum += Data.template Get<FTestFragment_Int>()[EntityIndex].Value;
+				}
+				Sum += ChunkSum;
+			});
+		}
+	};
+
+	struct FTestQueryExecutor_LoadTest : public UE::Mass::FQueryExecutor
+	{
+		UE::Mass::FQueryDefinition<
+			UE::Mass::FConstFragmentAccess<FTestFragment_Float>,
+			UE::Mass::FMutableFragmentAccess<FTestFragment_Int>,
+			UE::Mass::FConstFragmentAccess<FTestFragment_Bool>,
+			UE::Mass::FConstFragmentAccess<FTestFragment_Large>,
+			UE::Mass::FConstFragmentAccess<FTestFragment_Array>
+		> Accessors{ *this };
+
+		virtual void Execute(FMassExecutionContext& Context)
+		{
+			ForEachEntityChunk(Context, Accessors, [](FMassExecutionContext& Context, auto& Data)
+			{
+				for (uint32 EntityIndex : Context.CreateEntityIterator())
+				{
+					int32& TestInt = Data.template Get<FTestFragment_Int>()[EntityIndex].Value;
+
+					if (Data.template Get<FTestFragment_Float>()[EntityIndex].Value > 0.0f)
+					{
+						++TestInt;
+					}
+					else
+					{
+						--TestInt;
+					}
+
+					if (Data.template Get<FTestFragment_Bool>()[EntityIndex].bValue)
+					{
+						++TestInt;
+					}
+					else
+					{
+						--TestInt;
+					}
+
+					if (Data.template Get<FTestFragment_Large>()[EntityIndex].Value[0] > 0)
+					{
+						++TestInt;
+					}
+					else
+					{
+						--TestInt;
+					}
+
+					if (Data.template Get<FTestFragment_Array>()[EntityIndex].Value.Num() > 0)
+					{
+						++TestInt;
+					}
+					else
+					{
+						--TestInt;
+					}
+				}
+			});
+		}
+	};
+
+	struct FTestQueryExecutor_LoadTest_Parallel : public UE::Mass::FQueryExecutor
+	{
+		UE::Mass::FQueryDefinition<
+			UE::Mass::FConstFragmentAccess<FTestFragment_Float>,
+			UE::Mass::FMutableFragmentAccess<FTestFragment_Int>,
+			UE::Mass::FConstFragmentAccess<FTestFragment_Bool>,
+			UE::Mass::FConstFragmentAccess<FTestFragment_Large>,
+			UE::Mass::FConstFragmentAccess<FTestFragment_Array>
+		> Accessors{ *this };
+
+		virtual void Execute(FMassExecutionContext& Context)
+		{
+			ParallelForEachEntityChunk(Context, Accessors, [](FMassExecutionContext& Context, auto& Data)
+			{
+				for (uint32 EntityIndex : Context.CreateEntityIterator())
+				{
+					int32& TestInt = Data.template Get<FTestFragment_Int>()[EntityIndex].Value;
+
+					if (Data.template Get<FTestFragment_Float>()[EntityIndex].Value > 0.0f)
+					{
+						++TestInt;
+					}
+					else
+					{
+						--TestInt;
+					}
+
+					if (Data.template Get<FTestFragment_Bool>()[EntityIndex].bValue)
+					{
+						++TestInt;
+					}
+					else
+					{
+						--TestInt;
+					}
+
+					if (Data.template Get<FTestFragment_Large>()[EntityIndex].Value[0] > 0)
+					{
+						++TestInt;
+					}
+					else
+					{
+						--TestInt;
+					}
+
+					if (Data.template Get<FTestFragment_Array>()[EntityIndex].Value.Num() > 0)
+					{
+						++TestInt;
+					}
+					else
+					{
+						--TestInt;
+					}
+				}
+			});
+		}
+	};
+
+	struct FTestQueryExecutor_LoadTest_ByEntity : public UE::Mass::FQueryExecutor
+	{
+		UE::Mass::FQueryDefinition<
+			UE::Mass::FConstFragmentAccess<FTestFragment_Float>,
+			UE::Mass::FMutableFragmentAccess<FTestFragment_Int>,
+			UE::Mass::FConstFragmentAccess<FTestFragment_Bool>,
+			UE::Mass::FConstFragmentAccess<FTestFragment_Large>,
+			UE::Mass::FConstFragmentAccess<FTestFragment_Array>
+		> Accessors{ *this };
+
+		virtual void Execute(FMassExecutionContext& Context)
+		{
+			ForEachEntity(Context, Accessors, [](FMassExecutionContext& Context, auto& Data, uint32 EntityIndex)
+			{
+				int32& TestInt = Data.template Get<FTestFragment_Int>()[EntityIndex].Value;
+
+				if (Data.template Get<FTestFragment_Float>()[EntityIndex].Value > 0.0f)
+				{
+					++TestInt;
+				}
+				else
+				{
+					--TestInt;
+				}
+
+				if (Data.template Get<FTestFragment_Bool>()[EntityIndex].bValue)
+				{
+					++TestInt;
+				}
+				else
+				{
+					--TestInt;
+				}
+
+				if (Data.template Get<FTestFragment_Large>()[EntityIndex].Value[0] > 0)
+				{
+					++TestInt;
+				}
+				else
+				{
+					--TestInt;
+				}
+
+				if (Data.template Get<FTestFragment_Array>()[EntityIndex].Value.Num() > 0)
+				{
+					++TestInt;
+				}
+				else
+				{
+					--TestInt;
+				}
+			});
+		}
+	};
+
+	struct FTestQueryExecutor_LoadTest_ByEntity_Parallel : public UE::Mass::FQueryExecutor
+	{
+		UE::Mass::FQueryDefinition<
+			UE::Mass::FConstFragmentAccess<FTestFragment_Float>,
+			UE::Mass::FMutableFragmentAccess<FTestFragment_Int>,
+			UE::Mass::FConstFragmentAccess<FTestFragment_Bool>,
+			UE::Mass::FConstFragmentAccess<FTestFragment_Large>,
+			UE::Mass::FConstFragmentAccess<FTestFragment_Array>
+		> Accessors{ *this };
+
+		virtual void Execute(FMassExecutionContext& Context)
+		{
+			ParallelForEachEntity(Context, Accessors, [](FMassExecutionContext& Context, auto& Data, uint32 EntityIndex)
+			{
+				int32& TestInt = Data.template Get<FTestFragment_Int>()[EntityIndex].Value;
+
+				if (Data.template Get<FTestFragment_Float>()[EntityIndex].Value > 0.0f)
+				{
+					++TestInt;
+				}
+				else
+				{
+					--TestInt;
+				}
+
+				if (Data.template Get<FTestFragment_Bool>()[EntityIndex].bValue)
+				{
+					++TestInt;
+				}
+				else
+				{
+					--TestInt;
+				}
+
+				if (Data.template Get<FTestFragment_Large>()[EntityIndex].Value[0] > 0)
+				{
+					++TestInt;
+				}
+				else
+				{
+					--TestInt;
+				}
+
+				if (Data.template Get<FTestFragment_Array>()[EntityIndex].Value.Num() > 0)
+				{
+					++TestInt;
+				}
+				else
+				{
+					--TestInt;
+				}
+			});
+		}
+	};
+
+	struct FTestQueryExecutor_AnyTag : public UE::Mass::FQueryExecutor
+	{
+		int32 EntityCount = 0;
+
+		UE::Mass::FQueryDefinition<
+			UE::Mass::FConstFragmentAccess<FTestFragment_Float>
+		> Accessors{ *this };
+
+		virtual void Execute(FMassExecutionContext& Context)
+		{
+			ForEachEntity(Context, Accessors, [this](FMassExecutionContext& Context, auto& Data, uint32 EntityIndex)
+			{
+				++EntityCount;
+			});
+		}
+	};
+
+	struct FTestQueryExecutor_NeedTagA : public UE::Mass::FQueryExecutor
+	{
+		UE::Mass::FQueryDefinition<
+			UE::Mass::FConstFragmentAccess<FTestFragment_Float>,
+			UE::Mass::FMassTagRequired<FTestTag_A>
+		> Accessors{ *this };
+
+		int32 EntityCount = 0;
+
+		virtual void Execute(FMassExecutionContext& Context)
+		{
+			ForEachEntity(Context, Accessors, [this](FMassExecutionContext& Context, auto& Data, uint32 EntityIndex)
+			{
+				++EntityCount;
+			});
+		}
+	};
+
+	struct FTestQueryExecutor_BlockTagB : public UE::Mass::FQueryExecutor
+	{
+		UE::Mass::FQueryDefinition<
+			UE::Mass::FConstFragmentAccess<FTestFragment_Float>,
+			UE::Mass::FMassTagBlocked<FTestTag_B>
+		> Accessors{ *this };
+
+		int32 EntityCount = 0;
+
+		virtual void Execute(FMassExecutionContext& Context)
+		{
+			ForEachEntity(Context, Accessors, [this](FMassExecutionContext& Context, auto& Data, uint32 EntityIndex)
+			{
+				++EntityCount;
+			});
+		}
+	};
+
+	struct FQueryExecutor_IteratorConsistency : FEntityTestBase
+	{
+		virtual bool InstantTest() override
+		{
+			CA_ASSUME(EntityManager);
+
+			const int32 EntityCountNoTag = 7;
+			const int32 EntityCountA = 9;
+			const int32 EntityCountB = 13;
+			const int32 EntityCountAB = 17;
+
+			TArray<FMassEntityHandle> Entities;
+			EntityManager->BatchCreateEntities(FloatsArchetype, {}, EntityCountNoTag, Entities);
+
+			Entities.Reset();
+			EntityManager->BatchCreateEntities(FloatsArchetype, {}, EntityCountA, Entities);
+			for (FMassEntityHandle& Entity : Entities)
+			{
+				EntityManager->AddTagToEntity(Entity, FTestTag_A::StaticStruct());
+			}
+
+			Entities.Reset();
+			EntityManager->BatchCreateEntities(FloatsArchetype, {}, EntityCountB, Entities);
+			for (FMassEntityHandle& Entity : Entities)
+			{
+				EntityManager->AddTagToEntity(Entity, FTestTag_B::StaticStruct());
+			}
+
+			Entities.Reset();
+			EntityManager->BatchCreateEntities(FloatsArchetype, {}, EntityCountAB, Entities);
+			for (FMassEntityHandle& Entity : Entities)
+			{
+				EntityManager->AddTagToEntity(Entity, FTestTag_A::StaticStruct());
+				EntityManager->AddTagToEntity(Entity, FTestTag_B::StaticStruct());
+			}
+
+			TObjectPtr<UMassTestProcessorAutoExecuteQuery> Processor_AnyTag = NewObject<UMassTestProcessorAutoExecuteQuery>();
+			CA_ASSUME(Processor_AnyTag);
+			TSharedPtr<FTestQueryExecutor_AnyTag> TestQuery_AnyTag = UE::Mass::FQueryExecutor::CreateQuery<FTestQueryExecutor_AnyTag>(Processor_AnyTag->EntityQuery, Processor_AnyTag);
+			Processor_AnyTag->SetAutoExecuteQuery(TestQuery_AnyTag);
+			Processor_AnyTag->CallInitialize(GetTransientPackage(), EntityManager.ToSharedRef());
+
+			TObjectPtr<UMassTestProcessorAutoExecuteQuery> Processor_TagA = NewObject<UMassTestProcessorAutoExecuteQuery>();
+			CA_ASSUME(Processor_TagA);
+			TSharedPtr<FTestQueryExecutor_NeedTagA> TestQuery_TagA = UE::Mass::FQueryExecutor::CreateQuery<FTestQueryExecutor_NeedTagA>(Processor_TagA->EntityQuery, Processor_TagA);
+			Processor_TagA->SetAutoExecuteQuery(TestQuery_TagA);
+			Processor_TagA->CallInitialize(GetTransientPackage(), EntityManager.ToSharedRef());
+
+			TObjectPtr<UMassTestProcessorAutoExecuteQuery> Processor_TagB = NewObject<UMassTestProcessorAutoExecuteQuery>();
+			CA_ASSUME(Processor_TagB);
+			TSharedPtr<FTestQueryExecutor_BlockTagB> TestQuery_TagB = UE::Mass::FQueryExecutor::CreateQuery<FTestQueryExecutor_BlockTagB>(Processor_TagB->EntityQuery, Processor_TagB);
+			Processor_TagB->SetAutoExecuteQuery(TestQuery_TagB);
+			Processor_TagB->CallInitialize(GetTransientPackage(), EntityManager.ToSharedRef());
+
+			TArray<TObjectPtr<UMassProcessor>> Processors;
+			Processors.Add(Processor_AnyTag);
+			Processors.Add(Processor_TagA);
+			Processors.Add(Processor_TagB);
+
+			FMassProcessingContext ProcessingContext(EntityManager, /*DeltaSeconds=*/0.f);
+
+			UE::Mass::Executor::RunProcessorsView(Processors, ProcessingContext);
+
+			const int32 AnyCount = TestQuery_AnyTag->EntityCount;
+			const int32 ExpectedAnyCount = EntityCountNoTag + EntityCountA + EntityCountB + EntityCountAB;
+			const int32 ACount = TestQuery_TagA->EntityCount;
+			const int32 ExpectedACount = EntityCountA + EntityCountAB;
+			const int32 BCount = TestQuery_TagB->EntityCount;
+			const int32 ExpectedBCount = EntityCountNoTag + EntityCountA;
+
+			AITEST_EQUAL(TEXT("Any Tag Entities Processed"), AnyCount, ExpectedAnyCount);
+			AITEST_EQUAL(TEXT("Require TagA Entities Processed"), ACount, ExpectedACount);
+			AITEST_EQUAL(TEXT("Blocked TagB Entities Processed"), BCount, ExpectedBCount);
+
+			return true;
+		}
+	};
+	IMPLEMENT_AI_INSTANT_TEST(FQueryExecutor_IteratorConsistency, "System.Mass.Processor.AutoExecuteQuery.IteratorConsistency");
+
+	struct FTestQueryExecutor_LoadTest_Float : public UE::Mass::FQueryExecutor
+	{
+		UE::Mass::FQueryDefinition<
+			UE::Mass::FConstFragmentAccess<FTestFragment_Float>,
+			UE::Mass::FMutableFragmentAccess<FTestFragment_Int>
+		> Accessors{ *this };
+
+		virtual void Execute(FMassExecutionContext& Context)
+		{
+			ParallelForEachEntityChunk(Context, Accessors, [](FMassExecutionContext& Context, auto& Data)
+			{
+				for (uint32 EntityIndex : Context.CreateEntityIterator())
+				{
+					if (Data.template Get<FTestFragment_Float>()[EntityIndex].Value > 0.0f)
+					{
+						++Data.template Get<FTestFragment_Int>()[EntityIndex].Value;
+					}
+					else
+					{
+						--Data.template Get<FTestFragment_Int>()[EntityIndex].Value;
+					}
+				}
+			});
+		}
+	};
+
+	struct FTestQueryExecutor_LoadTest_Bool : public UE::Mass::FQueryExecutor
+	{
+		UE::Mass::FQueryDefinition<
+			UE::Mass::FMutableFragmentAccess<FTestFragment_Int>,
+			UE::Mass::FConstFragmentAccess<FTestFragment_Bool>
+		> Accessors{ *this };
+
+		virtual void Execute(FMassExecutionContext& Context)
+		{
+			ParallelForEachEntityChunk(Context, Accessors, [](FMassExecutionContext& Context, auto& Data)
+			{
+				for (uint32 EntityIndex : Context.CreateEntityIterator())
+				{
+					if (Data.template Get<FTestFragment_Bool>()[EntityIndex].bValue)
+					{
+						++Data.template Get<FTestFragment_Int>()[EntityIndex].Value;
+					}
+					else
+					{
+						--Data.template Get<FTestFragment_Int>()[EntityIndex].Value;
+					}
+				}
+			});
+		}
+	};
+
+	struct FTestQueryExecutor_LoadTest_Large : public UE::Mass::FQueryExecutor
+	{
+		UE::Mass::FQueryDefinition<
+			UE::Mass::FMutableFragmentAccess<FTestFragment_Int>,
+			UE::Mass::FConstFragmentAccess<FTestFragment_Large>
+		> Accessors{ *this };
+
+		virtual void Execute(FMassExecutionContext& Context)
+		{
+			ParallelForEachEntityChunk(Context, Accessors, [](FMassExecutionContext& Context, auto& Data)
+			{
+				for (uint32 EntityIndex : Context.CreateEntityIterator())
+				{
+					int32& TestInt = Data.template Get<FTestFragment_Int>()[EntityIndex].Value;
+
+					if (Data.template Get<FTestFragment_Large>()[EntityIndex].Value[0] > 0)
+					{
+						++Data.template Get<FTestFragment_Int>()[EntityIndex].Value;
+					}
+					else
+					{
+						--Data.template Get<FTestFragment_Int>()[EntityIndex].Value;
+					}
+				}
+			});
+		}
+	};
+
+	struct FTestQueryExecutor_LoadTest_Array : public UE::Mass::FQueryExecutor
+	{
+		UE::Mass::FQueryDefinition<
+			UE::Mass::FMutableFragmentAccess<FTestFragment_Int>,
+			UE::Mass::FConstFragmentAccess<FTestFragment_Array>
+		> Accessors{ *this };
+
+		virtual void Execute(FMassExecutionContext& Context)
+		{
+			ParallelForEachEntityChunk(Context, Accessors, [](FMassExecutionContext& Context, auto& Data)
+			{
+				for (uint32 EntityIndex : Context.CreateEntityIterator())
+				{
+					if (Data.template Get<FTestFragment_Array>()[EntityIndex].Value.Num() > 0)
+					{
+						++Data.template Get<FTestFragment_Int>()[EntityIndex].Value;
+					}
+					else
+					{
+						--Data.template Get<FTestFragment_Int>()[EntityIndex].Value;
+					}
+				}
+			});
+		}
+	};
+
+	struct FQueryExecutor_LoadTest : FEntityTestBase
+	{
+		virtual bool InstantTest() override
+		{
+			CA_ASSUME(EntityManager);
+
+			constexpr uint32 NumTags = 5;
+			const UScriptStruct* Tags[NumTags] = {
+				FTestFragment_Tag::StaticStruct(),
+				FTestTag_A::StaticStruct(),
+				FTestTag_B::StaticStruct(),
+				FTestTag_C::StaticStruct(),
+				FTestTag_D::StaticStruct(),
+			};
+
+			TArray<UScriptStruct*> Fragments;
+
+			Fragments.Add(FTestFragment_Float::StaticStruct());
+			Fragments.Add(FTestFragment_Int::StaticStruct());
+			Fragments.Add(FTestFragment_Bool::StaticStruct());
+			Fragments.Add(FTestFragment_Large::StaticStruct());
+			Fragments.Add(FTestFragment_Array::StaticStruct());
+
+			FMassArchetypeHandle LoadTestArchetype = EntityManager->CreateArchetype(Fragments);
+
+			TArray<FMassEntityHandle> Entities;
+
+			auto CreateEntities = [&](uint32 EntityCount)
+			{
+				EntityManager->BatchDestroyEntities(Entities);
+				Entities.Reset(EntityCount);
+				EntityManager->BatchCreateEntities(LoadTestArchetype, {}, EntityCount, Entities);
+
+				int32 TagMask = 0;
+				for (FMassEntityHandle& Entity : Entities)
+				{
+					// maximize entity fragmentation to test worst-case scenario
+
+					for (int32 BitIndex = 0; BitIndex < NumTags; ++BitIndex)
+					{
+						if ((TagMask & (1 << BitIndex)) != 0)
+						{
+							EntityManager->AddTagToEntity(Entity, Tags[BitIndex]);
+						}
+					}
+
+					if (++TagMask >= (1 << NumTags))
+					{
+						TagMask = 0;
+					}
+				}
+			};
+
+			TObjectPtr<UMassTestProcessorAutoExecuteQuery> Processor_LoadTest_Parallel = NewObject<UMassTestProcessorAutoExecuteQuery>();
+			CA_ASSUME(Processor_LoadTest_Parallel);
+			TSharedPtr<FTestQueryExecutor_LoadTest_Parallel> TestQuery_LoadTest_Parallel = UE::Mass::FQueryExecutor::CreateQuery<FTestQueryExecutor_LoadTest_Parallel>(Processor_LoadTest_Parallel->EntityQuery, Processor_LoadTest_Parallel);
+			Processor_LoadTest_Parallel->SetAutoExecuteQuery(TestQuery_LoadTest_Parallel);
+			Processor_LoadTest_Parallel->CallInitialize(GetTransientPackage(), EntityManager.ToSharedRef());
+
+			TObjectPtr<UMassTestProcessorAutoExecuteQuery> Processor_LoadTest_ByEntity_Parallel = NewObject<UMassTestProcessorAutoExecuteQuery>();
+			CA_ASSUME(Processor_LoadTest_ByEntity_Parallel);
+			TSharedPtr<FTestQueryExecutor_LoadTest_ByEntity_Parallel> TestQuery_LoadTest_ByEntity_Parallel = UE::Mass::FQueryExecutor::CreateQuery<FTestQueryExecutor_LoadTest_ByEntity_Parallel>(Processor_LoadTest_ByEntity_Parallel->EntityQuery, Processor_LoadTest_ByEntity_Parallel);
+			Processor_LoadTest_ByEntity_Parallel->SetAutoExecuteQuery(TestQuery_LoadTest_ByEntity_Parallel);
+			Processor_LoadTest_ByEntity_Parallel->CallInitialize(GetTransientPackage(), EntityManager.ToSharedRef());
+
+			TObjectPtr<UMassTestProcessorAutoExecuteQuery> Processor_LoadTest = NewObject<UMassTestProcessorAutoExecuteQuery>();
+			CA_ASSUME(Processor_LoadTest);
+			TSharedPtr<FTestQueryExecutor_LoadTest> TestQuery_LoadTest = UE::Mass::FQueryExecutor::CreateQuery<FTestQueryExecutor_LoadTest>(Processor_LoadTest->EntityQuery, Processor_LoadTest);
+			Processor_LoadTest->SetAutoExecuteQuery(TestQuery_LoadTest);
+			Processor_LoadTest->CallInitialize(GetTransientPackage(), EntityManager.ToSharedRef());
+
+			TObjectPtr<UMassTestProcessorAutoExecuteQuery> Processor_LoadTest_ByEntity = NewObject<UMassTestProcessorAutoExecuteQuery>();
+			CA_ASSUME(Processor_LoadTest_ByEntity);
+			TSharedPtr<FTestQueryExecutor_LoadTest_ByEntity> TestQuery_LoadTest_ByEntity = UE::Mass::FQueryExecutor::CreateQuery<FTestQueryExecutor_LoadTest_ByEntity>(Processor_LoadTest_ByEntity->EntityQuery, Processor_LoadTest_ByEntity);
+			Processor_LoadTest_ByEntity->SetAutoExecuteQuery(TestQuery_LoadTest_ByEntity);
+			Processor_LoadTest_ByEntity->CallInitialize(GetTransientPackage(), EntityManager.ToSharedRef());
+
+			TObjectPtr<UMassTestProcessorAutoExecuteQuery> Processor_Setup = NewObject<UMassTestProcessorAutoExecuteQuery>();
+			CA_ASSUME(Processor_Setup);
+			TSharedPtr<FTestQueryExecutor_LoadTestSetup> TestQuery_Setup = UE::Mass::FQueryExecutor::CreateQuery<FTestQueryExecutor_LoadTestSetup>(Processor_Setup->EntityQuery, Processor_Setup);
+			Processor_Setup->SetAutoExecuteQuery(TestQuery_Setup);
+			Processor_Setup->CallInitialize(GetTransientPackage(), EntityManager.ToSharedRef()); 
+
+			TObjectPtr<UMassTestProcessorAutoExecuteQuery> Processor_Reset = NewObject<UMassTestProcessorAutoExecuteQuery>();
+			CA_ASSUME(Processor_Reset);
+			TSharedPtr<FTestQueryExecutor_LoadTestReset> TestQuery_Reset = UE::Mass::FQueryExecutor::CreateQuery<FTestQueryExecutor_LoadTestReset>(Processor_Reset->EntityQuery, Processor_Reset);
+			Processor_Reset->SetAutoExecuteQuery(TestQuery_Reset);
+			Processor_Reset->CallInitialize(GetTransientPackage(), EntityManager.ToSharedRef());
+
+			TObjectPtr<UMassTestProcessorAutoExecuteQuery> Processor_Validate = NewObject<UMassTestProcessorAutoExecuteQuery>();
+			CA_ASSUME(Processor_Validate);
+			TSharedPtr<FTestQueryExecutor_LoadTestValidate> TestQuery_Validate = UE::Mass::FQueryExecutor::CreateQuery<FTestQueryExecutor_LoadTestValidate>(Processor_Validate->EntityQuery, Processor_Validate);
+			Processor_Validate->SetAutoExecuteQuery(TestQuery_Validate);
+			Processor_Validate->CallInitialize(GetTransientPackage(), EntityManager.ToSharedRef());
+
+			TObjectPtr<UMassTestProcessorAutoExecuteQuery> Processor_Float = NewObject<UMassTestProcessorAutoExecuteQuery>();
+			CA_ASSUME(Processor_Float);
+			TSharedPtr<FTestQueryExecutor_LoadTest_Float> TestQuery_Float = UE::Mass::FQueryExecutor::CreateQuery<FTestQueryExecutor_LoadTest_Float>(Processor_Float->EntityQuery, Processor_Float);
+			Processor_Float->SetAutoExecuteQuery(TestQuery_Float);
+			Processor_Float->CallInitialize(GetTransientPackage(), EntityManager.ToSharedRef());
+
+			TObjectPtr<UMassTestProcessorAutoExecuteQuery> Processor_Bool = NewObject<UMassTestProcessorAutoExecuteQuery>();
+			CA_ASSUME(Processor_Bool);
+			TSharedPtr<FTestQueryExecutor_LoadTest_Bool> TestQuery_Bool = UE::Mass::FQueryExecutor::CreateQuery<FTestQueryExecutor_LoadTest_Bool>(Processor_Bool->EntityQuery, Processor_Bool);
+			Processor_Bool->SetAutoExecuteQuery(TestQuery_Bool);
+			Processor_Bool->CallInitialize(GetTransientPackage(), EntityManager.ToSharedRef());
+
+			TObjectPtr<UMassTestProcessorAutoExecuteQuery> Processor_Large = NewObject<UMassTestProcessorAutoExecuteQuery>();
+			CA_ASSUME(Processor_Large);
+			TSharedPtr<FTestQueryExecutor_LoadTest_Large> TestQuery_Large = UE::Mass::FQueryExecutor::CreateQuery<FTestQueryExecutor_LoadTest_Large>(Processor_Large->EntityQuery, Processor_Large);
+			Processor_Large->SetAutoExecuteQuery(TestQuery_Large);
+			Processor_Large->CallInitialize(GetTransientPackage(), EntityManager.ToSharedRef());
+
+			TObjectPtr<UMassTestProcessorAutoExecuteQuery> Processor_Array = NewObject<UMassTestProcessorAutoExecuteQuery>();
+			CA_ASSUME(Processor_Array);
+			TSharedPtr<FTestQueryExecutor_LoadTest_Array> TestQuery_Array = UE::Mass::FQueryExecutor::CreateQuery<FTestQueryExecutor_LoadTest_Array>(Processor_Array->EntityQuery, Processor_Array);
+			Processor_Array->SetAutoExecuteQuery(TestQuery_Array);
+			Processor_Array->CallInitialize(GetTransientPackage(), EntityManager.ToSharedRef());
+
+			TArray<TObjectPtr<UMassProcessor>> IndividualProcessors_LoadTest;
+			IndividualProcessors_LoadTest.Add(Processor_Float);
+			IndividualProcessors_LoadTest.Add(Processor_Bool);
+			IndividualProcessors_LoadTest.Add(Processor_Large);
+			IndividualProcessors_LoadTest.Add(Processor_Array);
+
+			TArray<TObjectPtr<UMassProcessor>> Processors;
+			FMassProcessingContext ProcessingContext(EntityManager, /*DeltaSeconds=*/0.f);
+
+			UMassTestProcessorAutoExecuteQueryComparison* DefaultProcessor = NewTestProcessor<UMassTestProcessorAutoExecuteQueryComparison>(EntityManager);
+			UMassTestProcessorAutoExecuteQueryComparison_Parallel* DefaultProcessor_Parallel = NewTestProcessor<UMassTestProcessorAutoExecuteQueryComparison_Parallel>(EntityManager);
+
+			Processors.Reset();
+			Processors.Add(Processor_Setup);
+			UE::Mass::Executor::RunProcessorsView(Processors, ProcessingContext);
+
+			struct TestReturn
+			{
+				const TCHAR* Name;
+				uint64 Sum;
+				double Duration;
+				double NormalizedTime;
+			};
+
+			auto RunTest = [&](TObjectPtr<UMassProcessor> Proc, TArray<TObjectPtr<UMassProcessor>>* ProcArray, uint32 EntityCount, uint32 RunCount, const TCHAR* Name) -> TestReturn
+			{
+				Processors.Reset();
+				if (ProcArray != nullptr)
+				{
+					Processors.Append(*ProcArray);
+				}
+
+				if (Proc)
+				{
+					Processors.Add(Proc);
+				}
+			
+				// warm the cache with a single untimed iteration
+				UE::Mass::Executor::RunProcessorsView(Processors, ProcessingContext);
+
+				const double StartTime = FPlatformTime::Seconds();
+				for (uint32 i = 0; i < RunCount; i++)
+				{
+					UE::Mass::Executor::RunProcessorsView(Processors, ProcessingContext);
+				}
+				const double EndTime = FPlatformTime::Seconds();
+				const double Duration = EndTime - StartTime;
+				const double NormalizedTime = (Duration) / ((((double)EntityCount)/1000.0) * (((double)RunCount)/1000.0));
+				Processors.Reset();
+				Processors.Add(Processor_Validate);
+				UE::Mass::Executor::RunProcessorsView(Processors, ProcessingContext);
+				uint64 ProcSum = TestQuery_Validate->Sum;
+				TestQuery_Validate->Sum = 0;
+
+				Processors.Reset();
+				Processors.Add(Processor_Reset);
+				UE::Mass::Executor::RunProcessorsView(Processors, ProcessingContext);
+
+				return { Name, ProcSum, Duration, NormalizedTime };
+			};
+
+
+			auto LogResult = [&](TestReturn& Test, const TestReturn& Baseline)
+			{
+				const double PercentDiff = ((Test.NormalizedTime - Baseline.NormalizedTime) / Baseline.NormalizedTime) * 100.0;
+				if (PercentDiff > 0.0)
+				{
+					UE_LOG(LogMassPerfTest, Log, TEXT("%s: %.4fns/Entity (-%.4f%%)."), Test.Name, Test.NormalizedTime, PercentDiff);
+				}
+				else if (PercentDiff <= 0.0)
+				{
+					UE_LOG(LogMassPerfTest, Log, TEXT("%s: %.4fns/Entity (+%.4f%%)."), Test.Name, Test.NormalizedTime, FMath::Abs(PercentDiff));
+				}
+			};
+
+			auto LogTime = [&](TestReturn& Test)
+			{
+				UE_LOG(LogMassPerfTest, Log, TEXT("%s,%.8f"), Test.Name, Test.NormalizedTime);
+			};
+
+			auto SortPredicate = [](const TestReturn& A, const TestReturn& B)
+			{
+				return A.NormalizedTime < B.NormalizedTime;
+			};
+
+			auto RunBattery = [&](uint32 EntityCount, uint32 RunCount)
+			{
+				UE_LOG(LogMassPerfTest, Log, TEXT("\n%d entities:"), EntityCount);
+				CreateEntities(EntityCount);
+				TArray<TestReturn> Results;
+				int32 DefaultIndex = Results.Emplace(RunTest(DefaultProcessor, nullptr, EntityCount, RunCount, TEXT("Default")));
+				TestReturn Default = Results[DefaultIndex];
+				Results.Emplace(RunTest(nullptr, &IndividualProcessors_LoadTest, EntityCount, RunCount, TEXT("Split_Processing_Parallel")));
+				Results.Emplace(RunTest(Processor_LoadTest, nullptr, EntityCount, RunCount, TEXT("QueryExecutor")));
+				Results.Emplace(RunTest(Processor_LoadTest_ByEntity, nullptr, EntityCount, RunCount, TEXT("QueryExecutor_ByEntity")));
+				Results.Emplace(RunTest(Processor_LoadTest_Parallel, nullptr, EntityCount, RunCount, TEXT("QueryExecutor_Parallel")));
+				Results.Emplace(RunTest(Processor_LoadTest_ByEntity_Parallel, nullptr, EntityCount, RunCount, TEXT("QueryExecutor_ByEntity_Parallel")));
+				Results.Emplace(RunTest(DefaultProcessor_Parallel, nullptr, EntityCount, RunCount, TEXT("Default_Parallel")));
+
+				for (int32 i = 0; i < Results.Num(); i++)
+				{
+					FString TestMessage = FString::Format(TEXT("{0} Sum should match Default"), { Results[i].Name});
+					AITEST_EQUAL(TestMessage, Results[i].Sum, Default.Sum);
+				}
+
+
+				for (TestReturn& Result : Results)
+				{
+					LogTime(Result);
+				}
+
+				return true;
+			};
+			
+			RunBattery(1, 10000);
+			RunBattery(10, 1000);
+			RunBattery(100, 100);
+			RunBattery(1000, 100);
+
+			return true;
+		}
+	};
+	// This test can't fail and only exists to compare relative performance of various processor implementations and environments
+	//IMPLEMENT_AI_INSTANT_TEST(FQueryExecutor_LoadTest, "System.Mass.Processor.AutoExecuteQuery.LoadTest");
+
+} // FMassQueryExecutorTest
+
+UE_ENABLE_OPTIMIZATION_SHIP
+
+#undef LOCTEXT_NAMESPACE
diff --git a/Engine/Source/Developer/MassEntityTestSuite/Private/MassTestTypes.cpp b/Engine/Source/Developer/MassEntityTestSuite/Private/MassTestTypes.cpp
--- a/Engine/Source/Developer/MassEntityTestSuite/Private/MassTestTypes.cpp
+++ b/Engine/Source/Developer/MassEntityTestSuite/Private/MassTestTypes.cpp
@@ -118,6 +118,11 @@ void UMassTestProcessor_FloatsInts::ConfigureQueries(const TSharedRef<FMassEntit
 	 ExecutionFlags = int32(EProcessorExecutionFlags::All);
  }
 
+ UMassTestProcessorAutoExecuteQuery::UMassTestProcessorAutoExecuteQuery()
+ {
+	 bAutoRegisterWithProcessingPhases = false;
+ }
+
 //----------------------------------------------------------------------//
 // UMassTestWorldSubsystem
 //----------------------------------------------------------------------//
diff --git a/Engine/Source/Developer/MassEntityTestSuite/Public/MassEntityTestTypes.h b/Engine/Source/Developer/MassEntityTestSuite/Public/MassEntityTestTypes.h
--- a/Engine/Source/Developer/MassEntityTestSuite/Public/MassEntityTestTypes.h
+++ b/Engine/Source/Developer/MassEntityTestSuite/Public/MassEntityTestTypes.h
@@ -327,6 +327,39 @@ public:
 	static UE_API int StaticCounter;
 };
 
+UCLASS()
+class UMassTestProcessorAutoExecuteQuery : public UMassProcessor
+{
+	GENERATED_BODY()
+public:
+	UMassTestProcessorAutoExecuteQuery();
+	void SetAutoExecuteQuery(TSharedPtr<UE::Mass::FQueryExecutor> InAutoExecuteQuery)
+	{
+		AutoExecuteQuery = InAutoExecuteQuery;
+	}
+	FMassEntityQuery EntityQuery{ *this };
+};
+
+UCLASS()
+class UMassTestProcessorAutoExecuteQueryComparison : public UMassTestProcessorBase
+{
+	GENERATED_BODY()
+public:
+	virtual void ConfigureQueries(const TSharedRef<FMassEntityManager>&) override;
+	virtual void Execute(FMassEntityManager& EntityManager, FMassExecutionContext& Context) override;
+	FMassEntityQuery EntityQuery{ *this };
+};
+
+UCLASS()
+class UMassTestProcessorAutoExecuteQueryComparison_Parallel : public UMassTestProcessorBase
+{
+	GENERATED_BODY()
+public:
+	virtual void ConfigureQueries(const TSharedRef<FMassEntityManager>&) override;
+	virtual void Execute(FMassEntityManager& EntityManager, FMassExecutionContext& Context) override;
+	FMassEntityQuery EntityQuery{ *this };
+};
+
 struct FExecutionTestBase : FAITestBase
 {
 	TSharedPtr<FMassEntityManager> EntityManager;
diff --git a/Engine/Source/Runtime/MassEntity/Private/MassProcessor.cpp b/Engine/Source/Runtime/MassEntity/Private/MassProcessor.cpp
--- a/Engine/Source/Runtime/MassEntity/Private/MassProcessor.cpp
+++ b/Engine/Source/Runtime/MassEntity/Private/MassProcessor.cpp
@@ -8,6 +8,7 @@
 #include "MassCommandBuffer.h"
 #include "MassExecutionContext.h"
 #include "ProfilingDebugging/CsvProfiler.h"
+#include "MassQueryExecutor.h"
 
 #include UE_INLINE_GENERATED_CPP_BY_NAME(MassProcessor)
 
@@ -195,9 +196,16 @@ void UMassProcessor::InitializeInternal(UObject&, const TSharedRef<FMassEntityMa
 
 void UMassProcessor::ConfigureQueries(const TSharedRef<FMassEntityManager>&)
 {
-	UE_CVLOG_UELOG(OwnedQueries.Num(), this, LogMass, Warning
-		, TEXT("%s has entity queries registered. Make sure to override ConfigureQueries to configure the queries, and do not call the Super implementation")
-		, *GetProcessorName());
+	if (AutoExecuteQuery.IsValid())
+	{
+		AutoExecuteQuery->ConfigureQuery(ProcessorRequirements);
+	}
+	else
+	{
+		UE_CVLOG_UELOG(OwnedQueries.Num(), this, LogMass, Warning
+			, TEXT("%s has entity queries registered. Make sure to override ConfigureQueries to configure the queries, and do not call the Super implementation")
+			, *GetProcessorName());
+	}
 }
 
 void UMassProcessor::SetShouldAutoRegisterWithGlobalList(const bool bAutoRegister)
@@ -289,6 +297,19 @@ void UMassProcessor::CallExecute(FMassEntityManager& EntityManager, FMassExecuti
 	}
 }
 
+void UMassProcessor::Execute(FMassEntityManager& EntityManager, FMassExecutionContext& Context)
+{
+	if (AutoExecuteQuery.IsValid())
+	{
+		AutoExecuteQuery->CallExecute(Context);
+	}
+	else
+	{
+		static constexpr TCHAR MessageFormat[] = TEXT("UMassProcessor::Execute should never be called without an AutoExecuteQuery set. Override the function or populate AutoExecuteQuery. Processor name: %s");
+		checkf(false, MessageFormat, *GetProcessorName());
+	}
+}
+
 bool UMassProcessor::ShouldAllowQueryBasedPruning(const bool bRuntimeMode) const
 {
 	return bRuntimeMode && QueryBasedPruning == EMassQueryBasedPruning::Prune;
diff --git a/Engine/Source/Runtime/MassEntity/Private/MassQueryExecutor.cpp b/Engine/Source/Runtime/MassEntity/Private/MassQueryExecutor.cpp
new file mode 100644
--- /dev/null
+++ b/Engine/Source/Runtime/MassEntity/Private/MassQueryExecutor.cpp
@@ -0,0 +1,52 @@
+// Copyright Epic Games, Inc. All Rights Reserved.
+
+#include "MassQueryExecutor.h"
+#include "MassExecutionContext.h"
+
+namespace UE::Mass
+{
+
+FQueryExecutor::FQueryExecutor(FMassEntityQuery& InQuery, UObject* InLogOwner)
+	: BoundQuery(&InQuery)
+	, LogOwner(InLogOwner)
+{
+}
+
+FMassEntityQuery FQueryExecutor::DummyQuery;
+
+FQueryExecutor::FQueryExecutor()
+	: BoundQuery(&DummyQuery)
+{
+}
+
+void FQueryExecutor::CallExecute(FMassExecutionContext& Context)
+{
+#if WITH_MASSENTITY_DEBUG
+	ValidateAccessors();
+#endif
+
+	AccessorsPtr->SetupForExecute(Context);
+
+	Execute(Context);
+}
+
+void FQueryExecutor::ConfigureQuery(FMassSubsystemRequirements& ProcessorRequirements)
+{
+	AccessorsPtr->ConfigureQuery(*BoundQuery, ProcessorRequirements);
+}
+
+#if WITH_MASSENTITY_DEBUG
+void FQueryExecutor::ValidateAccessors()
+{
+	const uintptr_t ExecutorStart = (uintptr_t)this;
+	const uintptr_t ExecutorEnd = ExecutorStart + DebugSize;
+
+	if (AccessorsPtr)
+	{
+		const uintptr_t AccessorsStart = (uintptr_t)this;
+		checkf(ExecutorStart <= AccessorsStart && AccessorsStart <= ExecutorEnd, TEXT("Accessors assigned to a FQueryExecutor must be member variables of that struct."));
+	}
+};
+#endif //WITH_MASSENTITY_DEBUG
+
+} // namespace UE::Mass
\ No newline at end of file
diff --git a/Engine/Source/Runtime/MassEntity/Public/MassProcessor.h b/Engine/Source/Runtime/MassEntity/Public/MassProcessor.h
--- a/Engine/Source/Runtime/MassEntity/Public/MassProcessor.h
+++ b/Engine/Source/Runtime/MassEntity/Public/MassProcessor.h
@@ -18,6 +18,10 @@ class UMassCompositeProcessor;
 struct FMassDebugger;
 struct FMassEntityQuery;
 struct FMassEntityManager;
+namespace UE::Mass
+{
+	struct FQueryExecutor;
+}
 
 enum class EProcessorCompletionStatus : uint8
 {
@@ -192,7 +196,12 @@ protected:
 	UE_API virtual void ConfigureQueries(const TSharedRef<FMassEntityManager>& EntityManager);
 
 	UE_API virtual void PostInitProperties() override;
-	UE_API virtual void Execute(FMassEntityManager& EntityManager, FMassExecutionContext& Context) PURE_VIRTUAL(UMassProcessor::Execute);
+
+	/**
+	 * Called during the processing phase to which this processor is registered.
+	 * Default implementation requires that AutoExecuteQuery is populated with a QueryExecutor.
+	 */
+	UE_API virtual void Execute(FMassEntityManager& EntityManager, FMassExecutionContext& Context);
 
 protected:
 	/** Configures when this given processor can be executed in relation to other processors and processing groups, within its processing phase. */
@@ -278,6 +287,9 @@ protected:
 	/** A query representing elements this processor is accessing in Execute function outside of query execution */
 	FMassSubsystemRequirements ProcessorRequirements;
 
+	/** A QueryExecutor that can optionally be run in lieu of overriding the Execute function. */
+	TSharedPtr<UE::Mass::FQueryExecutor> AutoExecuteQuery;
+
 private:
 	/** Stores processor's queries registered via RegisterQuery. 
 	 *  @note that it's safe to store pointers here since RegisterQuery does verify that a given registered query is 
diff --git a/Engine/Source/Runtime/MassEntity/Public/MassQueryExecutor.h b/Engine/Source/Runtime/MassEntity/Public/MassQueryExecutor.h
new file mode 100644
--- /dev/null
+++ b/Engine/Source/Runtime/MassEntity/Public/MassQueryExecutor.h
@@ -0,0 +1,795 @@
+// Copyright Epic Games, Inc. All Rights Reserved.
+
+#pragma once
+
+#include "Misc/NotNull.h"
+#include "MassExecutionContext.h"
+
+class UMassProcessor;
+
+namespace UE::Mass
+{
+
+/** Interface for QueryDefinition templates. Not intended for other direct inheritance. */
+struct FQueryDefinitionBase
+{
+	virtual void ConfigureQuery(FMassEntityQuery& EntityQuery, FMassSubsystemRequirements& ProcessorRequirements) = 0;
+	virtual void SetupForExecute(FMassExecutionContext& Context) = 0;
+};
+
+/**
+ * A MassEntityQuery wrapper with type-safe data access.
+ */
+struct MASSENTITY_API FQueryExecutor
+{
+	template<typename T, typename = typename TEnableIf<TIsDerivedFrom<T, FQueryExecutor>::IsDerived>::Type>
+	static TSharedPtr<T> CreateQuery(FMassEntityQuery& InQuery, UObject* InLogOwner = nullptr)
+	{
+		TSharedPtr<T> Query = MakeShared<T>();
+		Query->BoundQuery = &InQuery;
+		Query->LogOwner = InLogOwner;
+
+#if WITH_MASSENTITY_DEBUG
+		Query->DebugSize = sizeof(T);
+		Query->ValidateAccessors();
+#endif
+
+		return Query;
+	}
+
+	template<typename... Ts>
+	friend struct FQueryDefinition;
+
+	friend ::UMassProcessor;
+
+	explicit FQueryExecutor(FMassEntityQuery& InQuery, UObject* InLogOwner = nullptr);
+	virtual ~FQueryExecutor() = default;
+
+	/** Override with logic to perform against the entities returned by this query. */
+	virtual void Execute(FMassExecutionContext& Context) = 0;
+
+protected:
+
+	FORCEINLINE UObject* GetLogOwner()
+	{
+		return LogOwner.Get();
+	}
+
+	template<typename TAccessors, typename TFunc>
+	FORCEINLINE void ForEachEntityChunk(FMassExecutionContext& ExecutionContext, TAccessors& Accessors, const TFunc&& ExecuteFunction)
+	{
+		BoundQuery->ForEachEntityChunk(ExecutionContext, [&Accessors, &ExecuteFunction](FMassExecutionContext& Context)
+		{
+			Accessors.AccessorTuple.ApplyBefore([&Context](auto&... Accessor)
+			{
+				(Accessor.SetupForChunk(Context), ...);
+			});
+			ExecuteFunction(Context, Accessors);
+		});
+	}
+
+	template<typename TAccessors, typename TFunc>
+	FORCEINLINE void ParallelForEachEntityChunk(FMassExecutionContext& ExecutionContext, const TAccessors& Accessors, const TFunc&& ExecuteFunction)
+	{
+		BoundQuery->ParallelForEachEntityChunk(ExecutionContext, [&Accessors, &ExecuteFunction](FMassExecutionContext& Context)
+		{
+			TAccessors LocalAccessors = TAccessors(Accessors);
+
+			LocalAccessors.AccessorTuple.ApplyBefore([&Context](auto&... Accessor)
+			{
+				(Accessor.SetupForChunk(Context), ...);
+			});
+
+			ExecuteFunction(Context, LocalAccessors);
+		});
+	}
+
+	template<typename TAccessors, typename TFunc>
+	FORCEINLINE void ForEachEntity(FMassExecutionContext& ExecutionContext, TAccessors& Accessors, const TFunc&& ExecuteFunction)
+	{
+		BoundQuery->ForEachEntityChunk(ExecutionContext, [&Accessors, &ExecuteFunction](FMassExecutionContext& Context)
+		{
+			Accessors.AccessorTuple.ApplyBefore([&Context](auto&... Accessor)
+			{
+				(Accessor.SetupForChunk(Context), ...);
+			});
+
+			for (uint32 EntityIndex : Context.CreateEntityIterator())
+			{
+				ExecuteFunction(Context, Accessors, EntityIndex);
+			}
+		});
+	}
+
+	template<typename TAccessors, typename TFunc>
+	FORCEINLINE void ParallelForEachEntity(FMassExecutionContext& ExecutionContext, TAccessors& Accessors, const TFunc&& ExecuteFunction)
+	{
+		BoundQuery->ParallelForEachEntityChunk(ExecutionContext, [&Accessors, &ExecuteFunction](FMassExecutionContext& Context)
+		{
+			TAccessors LocalAccessors = TAccessors(Accessors);
+			LocalAccessors.AccessorTuple.ApplyBefore([&Context](auto&... Accessor)
+			{
+				(Accessor.SetupForChunk(Context), ...);
+			});
+
+			for (uint32 EntityIndex : Context.CreateEntityIterator())
+			{
+				ExecuteFunction(Context, LocalAccessors, EntityIndex);
+			}
+		});
+	}
+
+	FQueryExecutor();
+private:
+	static FMassEntityQuery DummyQuery;
+
+	void ConfigureQuery(FMassSubsystemRequirements& ProcessorRequirements);
+	void CallExecute(FMassExecutionContext& Context);
+
+	TNotNull<FMassEntityQuery*> BoundQuery;
+	TWeakObjectPtr<UObject> LogOwner;
+
+	/** AccessorsPtr is only allowed to point to a member variable of this struct and will assert in all other cases. */
+	FQueryDefinitionBase* AccessorsPtr;
+
+#if WITH_MASSENTITY_DEBUG
+	void ValidateAccessors();
+	uint32 DebugSize = 0;
+#endif
+};
+
+template <typename...>
+inline constexpr auto IsUnique = std::true_type{};
+
+template <typename T, typename... Rest>
+inline constexpr auto IsUnique<T, Rest...> = std::bool_constant<(!std::is_same_v<typename T::FragmentType, typename Rest::FragmentType> && ...) && IsUnique<Rest...>>{};
+
+
+template <typename T, typename U, typename... Us>
+constexpr auto GetIndexRecursive()
+{
+	if constexpr (std::is_same<T, typename U::FragmentType>::value)
+	{
+		return 0;
+	}
+	else
+	{
+		static_assert(sizeof...(Us) > 0, "Type not found in accessor collection.");
+		return 1 + GetIndexRecursive<T, Us...>();
+	}
+}
+
+template <typename T, typename U, typename... Us>
+constexpr auto GetAccessorIndex()
+{
+	static_assert(IsUnique<U, Us...>, "An accessor collection must only contain a single instance of each fragment/subsystem/tag type.");
+	return GetIndexRecursive<T, U, Us...>();
+}
+
+/**
+ * Defines the entity compositions to return in the query and provides type-safe data access to
+ * entity and subsystem data. Must be a member variable of a QueryExecutor
+ */
+template<typename... Accessors>
+struct FQueryDefinition : public FQueryDefinitionBase
+{
+	FQueryDefinition(FQueryExecutor& Owner)
+	{
+		Owner.AccessorsPtr = this;
+	}
+
+	TTuple<Accessors...> AccessorTuple{};
+
+	virtual void ConfigureQuery(FMassEntityQuery& EntityQuery, FMassSubsystemRequirements& ProcessorRequirements) override
+	{
+		AccessorTuple.ApplyBefore([&EntityQuery, &ProcessorRequirements](auto&... Accessor)
+		{
+			(Accessor.ConfigureQuery(EntityQuery, ProcessorRequirements), ...);
+		});
+
+	}
+
+	virtual void SetupForExecute(FMassExecutionContext& Context) override
+	{
+		AccessorTuple.ApplyBefore([&Context](auto&... Accessor)
+		{
+			(Accessor.SetupForExecute(Context), ...);
+		});
+	}
+
+	FORCEINLINE void SetupForChunk(FMassExecutionContext& Context)
+	{
+		AccessorTuple.ApplyBefore([&Context](auto&... Accessor)
+		{
+			(Accessor.SetupForChunk(Context), ...);
+		});
+	}
+
+	template <typename TFragment>
+	FORCEINLINE constexpr auto& Get()
+	{
+		constexpr std::size_t Index = GetAccessorIndex<TFragment, Accessors...>();
+		return AccessorTuple.template Get<Index>();
+	}
+};
+
+template<typename TFragment>
+struct FMutableFragmentAccess
+{
+	template<typename... Ts>
+	friend struct FQueryDefinition;
+
+	using FragmentType = TFragment;
+
+	FMutableFragmentAccess() = default;
+
+	FORCEINLINE TArrayView<TFragment>& Get()
+	{
+		return View;
+	}
+
+	FORCEINLINE TFragment& operator[](int32 Index)
+	{
+		return View[Index];
+	}
+
+	FORCEINLINE void ConfigureQuery(FMassEntityQuery& EntityQuery, FMassSubsystemRequirements& ProcessorRequirements)
+	{
+		EntityQuery.AddRequirement<TFragment>(EMassFragmentAccess::ReadWrite);
+	}
+
+	FORCEINLINE void SetupForExecute(FMassExecutionContext& Context)
+	{
+	}
+
+	FORCEINLINE void SetupForChunk(FMassExecutionContext& Context)
+	{
+		View = Context.GetMutableFragmentView<TFragment>();
+	}
+
+	TArrayView<TFragment> View;
+};
+
+template<typename TFragment>
+struct FConstFragmentAccess
+{
+	template<typename... Ts>
+	friend struct FQueryDefinition;
+
+	using FragmentType = TFragment;
+
+	FConstFragmentAccess() = default;
+
+	FORCEINLINE const TConstArrayView<TFragment>& Get() const
+	{
+		return View;
+	}
+
+	FORCEINLINE const TFragment& operator[](int32 Index) const
+	{
+		return View[Index];
+	}
+
+	FORCEINLINE void ConfigureQuery(FMassEntityQuery& EntityQuery, FMassSubsystemRequirements& ProcessorRequirements) const
+	{
+		EntityQuery.AddRequirement<TFragment>(EMassFragmentAccess::ReadOnly);
+	}
+
+	FORCEINLINE void SetupForExecute(FMassExecutionContext& Context)
+	{
+	}
+
+	FORCEINLINE void SetupForChunk(FMassExecutionContext& Context)
+	{
+		View = Context.GetFragmentView<TFragment>();
+	}
+
+	TConstArrayView<TFragment> View;
+};
+
+template<typename TFragment>
+struct FMutableOptionalFragmentAccess
+{
+	template<typename... Ts>
+	friend struct FQueryDefinition;
+
+	using FragmentType = TFragment;
+
+	FMutableOptionalFragmentAccess() = default;
+
+	FORCEINLINE TArrayView<TFragment>& Get() const
+	{
+		return View;
+	}
+
+	FORCEINLINE TFragment& operator[](int32 Index) const
+	{
+		return View[Index];
+	}
+
+	FORCEINLINE int32 Num() const
+	{
+		return View.Num();
+	}
+
+	FORCEINLINE operator bool() const
+	{
+		return View.Num() > 0;
+	}
+
+	FORCEINLINE void ConfigureQuery(FMassEntityQuery& EntityQuery, FMassSubsystemRequirements& ProcessorRequirements) const
+	{
+		EntityQuery.AddRequirement<TFragment>(EMassFragmentAccess::ReadWrite, EMassFragmentPresence::Optional);
+	}
+
+	FORCEINLINE void SetupForExecute(FMassExecutionContext& Context)
+	{
+	}
+
+	FORCEINLINE void SetupForChunk(FMassExecutionContext& Context)
+	{
+		View = Context.GetMutableFragmentView<TFragment>();
+	}
+
+	TArrayView<TFragment> View;
+};
+
+template<typename TFragment>
+struct FConstOptionalFragmentAccess
+{
+	template<typename... Ts>
+	friend struct FQueryDefinition;
+
+	using FragmentType = TFragment;
+
+	FConstOptionalFragmentAccess() = default;
+
+	FORCEINLINE const TConstArrayView<TFragment>& Get() const
+	{
+		return View;
+	}
+
+	FORCEINLINE const TFragment& operator[](int32 Index) const
+	{
+		return View[Index];
+	}
+
+	FORCEINLINE int32 Num() const
+	{
+		return View.Num();
+	}
+
+	FORCEINLINE operator bool() const
+	{
+		return View.Num() > 0;
+	}
+
+	FORCEINLINE void ConfigureQuery(FMassEntityQuery& EntityQuery, FMassSubsystemRequirements& ProcessorRequirements) const
+	{
+		EntityQuery.AddRequirement<TFragment>(EMassFragmentAccess::ReadOnly);
+	}
+
+	FORCEINLINE void SetupForExecute(FMassExecutionContext& Context)
+	{
+	}
+
+	FORCEINLINE void SetupForChunk(FMassExecutionContext& Context)
+	{
+		View = Context.GetFragmentView<TFragment>();
+	}
+
+	TConstArrayView<TFragment> View;
+};
+
+template<typename TTag>
+struct FMassTagRequired
+{
+	template<typename... Ts>
+	friend struct FQueryDefinition;
+
+	using FragmentType = TTag;
+
+	FMassTagRequired() = default;
+
+	FORCEINLINE void ConfigureQuery(FMassEntityQuery& EntityQuery, FMassSubsystemRequirements& ProcessorRequirements) const
+	{
+		EntityQuery.AddTagRequirement<TTag>(EMassFragmentPresence::All);
+	}
+
+	FORCEINLINE void SetupForExecute(FMassExecutionContext& Context)
+	{
+	}
+
+	FORCEINLINE void SetupForChunk(FMassExecutionContext& Context)
+	{
+	}
+};
+
+template<typename TFragment>
+struct FMassTagBlocked
+{
+	using FragmentType = TFragment;
+
+	FORCEINLINE void ConfigureQuery(FMassEntityQuery& EntityQuery, FMassSubsystemRequirements& ProcessorRequirements) const
+	{
+		EntityQuery.AddTagRequirement<TFragment>(EMassFragmentPresence::None);
+	}
+
+	FORCEINLINE void SetupForExecute(FMassExecutionContext& Context)
+	{
+	}
+
+	FORCEINLINE void SetupForChunk(FMassExecutionContext& Context)
+	{
+	}
+};
+
+template<typename TFragment>
+struct FMutableSharedFragmentAccess
+{
+	template<typename... Ts>
+	friend struct FQueryDefinition;
+
+	using FragmentType = TFragment;
+
+	FMutableSharedFragmentAccess() = default;
+
+	FORCEINLINE TFragment& Get() const
+	{
+		check(Fragment);
+		return *Fragment;
+	}
+
+	FORCEINLINE TFragment& operator*() const
+	{
+		check(Fragment);
+		return *Fragment;
+	}
+
+	FORCEINLINE TFragment* operator->() const
+	{
+		return Fragment;
+	}
+
+	FORCEINLINE void ConfigureQuery(FMassEntityQuery& EntityQuery, FMassSubsystemRequirements& ProcessorRequirements) const
+	{
+		EntityQuery.AddSharedRequirement<TFragment>(EMassFragmentAccess::ReadWrite);
+	}
+
+	FORCEINLINE void SetupForExecute(FMassExecutionContext& Context)
+	{
+	}
+
+	FORCEINLINE void SetupForChunk(FMassExecutionContext& Context)
+	{
+		Fragment = &(Context.GetSharedFragment<TFragment>());
+	}
+
+	TFragment* Fragment = nullptr;
+};
+
+template<typename TFragment>
+struct FConstSharedFragmentAccess
+{
+	template<typename... Ts>
+	friend struct FQueryDefinition;
+
+	using FragmentType = TFragment;
+
+	FConstSharedFragmentAccess() = default;
+
+	FORCEINLINE const TFragment& Get() const
+	{
+		check(Fragment);
+		return *Fragment;
+	}
+
+	FORCEINLINE const TFragment& operator*() const
+	{
+		check(Fragment);
+		return *Fragment;
+	}
+
+	FORCEINLINE const TFragment* operator->() const
+	{
+		return Fragment;
+	}
+
+	FORCEINLINE void ConfigureQuery(FMassEntityQuery& EntityQuery, FMassSubsystemRequirements& ProcessorRequirements) const
+	{
+		EntityQuery.AddConstSharedRequirement<TFragment>();
+	}
+
+	FORCEINLINE void SetupForExecute(FMassExecutionContext& Context)
+	{
+	}
+
+	FORCEINLINE void SetupForChunk(FMassExecutionContext& Context)
+	{
+		Fragment = &Context.GetConstSharedFragment<TFragment>();
+	}
+
+	const TFragment* Fragment = nullptr;
+};
+
+template<typename TFragment>
+struct FMutableChunkFragmentAccess
+{
+	template<typename... Ts>
+	friend struct FQueryDefinition;
+
+	using FragmentType = TFragment;
+
+	FMutableChunkFragmentAccess() = default;
+
+	FORCEINLINE TFragment& Get() const
+	{
+		check(Fragment);
+		return *Fragment;
+	}
+
+	FORCEINLINE TFragment& operator*() const
+	{
+		check(Fragment);
+		return *Fragment;
+	}
+
+	FORCEINLINE TFragment* operator->() const
+	{
+		return Fragment;
+	}
+
+	FORCEINLINE void ConfigureQuery(FMassEntityQuery& EntityQuery, FMassSubsystemRequirements& ProcessorRequirements) const
+	{
+		EntityQuery.AddChunkRequirement<TFragment>(EMassFragmentAccess::ReadWrite, EMassFragmentPresence::All);
+	}
+
+	FORCEINLINE void SetupForExecute(FMassExecutionContext& Context)
+	{
+	}
+
+	FORCEINLINE void SetupForChunk(FMassExecutionContext& Context)
+	{
+		Fragment = Context.GetMutableChunkFragmentPtr<TFragment>();
+	}
+
+	TFragment* Fragment = nullptr;
+};
+
+template<typename TFragment>
+struct FConstChunkFragmentAccess
+{
+	template<typename... Ts>
+	friend struct FQueryDefinition;
+
+	using FragmentType = TFragment;
+
+	FConstChunkFragmentAccess() = default;
+
+	FORCEINLINE const TFragment& Get() const
+	{
+		check(Fragment);
+		return *Fragment;
+	}
+
+	FORCEINLINE const TFragment& operator*() const
+	{
+		check(Fragment);
+		return *Fragment;
+	}
+
+	FORCEINLINE const TFragment* operator->() const
+	{
+		return Fragment;
+	}
+
+	FORCEINLINE void ConfigureQuery(FMassEntityQuery& EntityQuery, FMassSubsystemRequirements& ProcessorRequirements) const
+	{
+		EntityQuery.AddChunkRequirement<TFragment>(EMassFragmentAccess::ReadOnly, EMassFragmentPresence::All);
+	}
+
+	FORCEINLINE void SetupForExecute(FMassExecutionContext& Context)
+	{
+	}
+
+	FORCEINLINE void SetupForChunk(FMassExecutionContext& Context)
+	{
+		Fragment = Context.GetChunkFragmentPtr<TFragment>();
+	}
+
+	const TFragment* Fragment = nullptr;
+};
+
+template<typename TFragment>
+struct FMutableOptionalChunkFragmentAccess
+{
+	template<typename... Ts>
+	friend struct FQueryDefinition;
+
+	using FragmentType = TFragment;
+
+	FMutableOptionalChunkFragmentAccess() = default;
+
+	FORCEINLINE TFragment* Get() const
+	{
+		check(Fragment);
+		return Fragment;
+	}
+
+	FORCEINLINE TFragment& operator*() const
+	{
+		check(Fragment);
+		return *Fragment;
+	}
+
+	FORCEINLINE TFragment* operator->() const
+	{
+		return Fragment;
+	}
+
+	FORCEINLINE operator bool() { return Fragment != nullptr; }
+
+	FORCEINLINE void ConfigureQuery(FMassEntityQuery& EntityQuery, FMassSubsystemRequirements& ProcessorRequirements) const
+	{
+		EntityQuery.AddChunkRequirement<TFragment>(EMassFragmentAccess::ReadWrite, EMassFragmentPresence::Optional);
+	}
+
+	FORCEINLINE void SetupForExecute(FMassExecutionContext& Context)
+	{
+	}
+
+	FORCEINLINE void SetupForChunk(FMassExecutionContext& Context)
+	{
+		Fragment = Context.GetMutableChunkFragmentPtr<TFragment>();
+	}
+
+	TFragment* Fragment = nullptr;
+};
+
+template<typename TFragment>
+struct FConstOptionalChunkFragmentAccess
+{
+	template<typename... Ts>
+	friend struct FQueryDefinition;
+
+	using FragmentType = TFragment;
+
+	FConstOptionalChunkFragmentAccess() = default;
+
+	FORCEINLINE const TFragment* Get()
+	{
+		check(Fragment);
+		return Fragment;
+	}
+
+	FORCEINLINE const TFragment& operator*() const
+	{
+		check(Fragment);
+		return *Fragment;
+	}
+
+	FORCEINLINE const TFragment* operator->() const
+	{
+		return Fragment;
+	}
+
+	FORCEINLINE operator bool() const
+	{
+		return Fragment != nullptr;
+	}
+
+	FORCEINLINE void ConfigureQuery(FMassEntityQuery& EntityQuery, FMassSubsystemRequirements& ProcessorRequirements) const
+	{
+		EntityQuery.AddChunkRequirement<TFragment>(EMassFragmentAccess::ReadOnly, EMassFragmentPresence::Optional);
+	}
+
+	FORCEINLINE void SetupForExecute(FMassExecutionContext& Context)
+	{
+	}
+
+	FORCEINLINE void SetupForChunk(FMassExecutionContext& Context)
+	{
+		Fragment = Context.GetChunkFragmentPtr<TFragment>();
+	}
+
+	const TFragment* Fragment = nullptr;
+};
+
+template<typename TSubsystem, typename = typename TEnableIf<TIsDerivedFrom<TSubsystem, USubsystem>::IsDerived>::Type>
+struct FMutableSubsystemAccess
+{
+	template<typename... Ts>
+	friend struct FQueryDefinition;
+
+	using FragmentType = TSubsystem;
+
+	FMutableSubsystemAccess() = default;
+
+	FORCEINLINE TSubsystem* Get() const
+	{
+		return Subsystem;
+	}
+
+	FORCEINLINE TSubsystem& operator*() const
+	{
+		check(Subsystem);
+		return *Subsystem;
+	}
+
+	FORCEINLINE TSubsystem* operator->() const
+	{
+		return Subsystem;
+	}
+
+	FORCEINLINE operator bool() const
+	{
+		return Subsystem != nullptr;
+	}
+
+	FORCEINLINE void ConfigureQuery(FMassEntityQuery& EntityQuery, FMassSubsystemRequirements& ProcessorRequirements) const
+	{
+		ProcessorRequirements.AddSubsystemRequirement<TSubsystem>(EMassFragmentAccess::ReadWrite);
+	}
+
+	FORCEINLINE void SetupForExecute(FMassExecutionContext& Context)
+	{
+		Subsystem = Context.GetMutableSubsystem<TSubsystem>();
+	}
+
+	FORCEINLINE void SetupForChunk(FMassExecutionContext& Context)
+	{
+	}
+
+	TSubsystem* Subsystem = nullptr;
+};
+
+template<typename TSubsystem, typename = typename TEnableIf<TIsDerivedFrom<TSubsystem, USubsystem>::IsDerived>::Type>
+struct FConstSubsystemAccess
+{
+	template<typename... Ts>
+	friend struct FQueryDefinition;
+
+	using FragmentType = TSubsystem;
+
+	FConstSubsystemAccess() = default;
+
+	FORCEINLINE const TSubsystem* Get() const
+	{
+		return Subsystem;
+	}
+
+	FORCEINLINE const TSubsystem& operator*() const
+	{
+		check(Subsystem);
+		return *Subsystem;
+	}
+
+	FORCEINLINE const TSubsystem* operator->() const
+	{
+		return Subsystem;
+	}
+
+	FORCEINLINE operator bool() const
+	{
+		return Subsystem != nullptr;
+	}
+
+	FORCEINLINE void ConfigureQuery(FMassEntityQuery& EntityQuery, FMassSubsystemRequirements& ProcessorRequirements) const
+	{
+		ProcessorRequirements.AddSubsystemRequirement<TSubsystem>(EMassFragmentAccess::ReadOnly);
+	}
+
+	FORCEINLINE void SetupForExecute(FMassExecutionContext& Context)
+	{
+		Subsystem = Context.GetSubsystem(TSubsystem::StaticClass());
+	}
+
+	FORCEINLINE void SetupForChunk(FMassExecutionContext& Context)
+	{
+	}
+
+	const TSubsystem* Subsystem = nullptr;
+};
+
+
+} // namespace UE::Mass
