[MassEntityTestSuite]
Fixed mismatched optimization macros.

diff --git a/Engine/Source/Developer/MassEntityTestSuite/Private/MassQueryExecutorTest.cpp b/Engine/Source/Developer/MassEntityTestSuite/Private/MassQueryExecutorTest.cpp
--- a/Engine/Source/Developer/MassEntityTestSuite/Private/MassQueryExecutorTest.cpp
+++ b/Engine/Source/Developer/MassEntityTestSuite/Private/MassQueryExecutorTest.cpp
@@ -16,7 +16,7 @@
 
 DEFINE_LOG_CATEGORY_STATIC(LogMassPerfTest, Log, All);
 
-//UE_DISABLE_OPTIMIZATION_SHIP
+UE_DISABLE_OPTIMIZATION_SHIP
 
 void UMassTestProcessorAutoExecuteQueryComparison::ConfigureQueries(const TSharedRef<FMassEntityManager>& EntityManager)
 {
