[Mass] Fixed mass observers not being initialized in unit tests due to the FMassEntityManager::PostInitialize not being called

#rnx

diff --git a/Engine/Source/Developer/MassEntityTestSuite/Private/MassTestTypes.cpp b/Engine/Source/Developer/MassEntityTestSuite/Private/MassTestTypes.cpp
--- a/Engine/Source/Developer/MassEntityTestSuite/Private/MassTestTypes.cpp
+++ b/Engine/Source/Developer/MassEntityTestSuite/Private/MassTestTypes.cpp
@@ -19,6 +19,7 @@ bool FExecutionTestBase::SetUp()
 	EntityManager = MakeShareable(new FMassEntityManager(bMakeWorldEntityManagersOwner ? FAITestHelpers::GetWorld() : nullptr));
 	EntityManager->SetDebugName(TEXT("MassEntityTestSuite"));
 	EntityManager->Initialize();
+	EntityManager->PostInitialize();
 
 	return true;
 }
diff --git a/Engine/Source/Developer/MassEntityTestSuite/Public/MassEntityTestTypes.h b/Engine/Source/Developer/MassEntityTestSuite/Public/MassEntityTestTypes.h
--- a/Engine/Source/Developer/MassEntityTestSuite/Public/MassEntityTestTypes.h
+++ b/Engine/Source/Developer/MassEntityTestSuite/Public/MassEntityTestTypes.h
@@ -377,7 +377,7 @@ public:
 struct FExecutionTestBase : FAITestBase
 {
 	TSharedPtr<FMassEntityManager> EntityManager;
-	bool bMakeWorldEntityManagersOwner = false;
+	bool bMakeWorldEntityManagersOwner = true;
 
 	UE_API virtual bool SetUp() override;
 };
