[Mass] Fixed FarmTest to not crash on PIE
The existing implementation is crashing since the changes to processor initialization have not been applied there. This issue is limited to the farm-test, due to the custom and legacy way of setting up and running processors.

diff --git a/Engine/Source/Developer/MassEntityTestSuite/Private/MassEntityTestFarmPlot.cpp b/Engine/Source/Developer/MassEntityTestSuite/Private/MassEntityTestFarmPlot.cpp
--- a/Engine/Source/Developer/MassEntityTestSuite/Private/MassEntityTestFarmPlot.cpp
+++ b/Engine/Source/Developer/MassEntityTestSuite/Private/MassEntityTestFarmPlot.cpp
@@ -237,6 +237,7 @@ void AMassEntityTestFarmPlot::BeginPlay()
 	FMassArchetypeHandle FlowerArchetype = EntityManager.CreateArchetype(TArray<const UScriptStruct*>{ FFarmWaterFragment::StaticStruct(), FFarmFlowerFragment::StaticStruct(), FHarvestTimerFragment::StaticStruct(), FFarmVisualFragment::StaticStruct(), FFarmGridCellData::StaticStruct() });
 
 	PerFrameSystems.Add(NewObject<UFarmWaterProcessor>(this));
+	PerFrameSystems.Last()->CallInitialize(this, EntityManager.AsShared());
 
 	PerSecondSystems.Add(NewObject<UFarmHarvestTimerSystem_Flowers>(this));
 	PerSecondSystems.Add(NewObject<UFarmHarvestTimerSystem_Crops>(this));
@@ -250,6 +251,11 @@ void AMassEntityTestFarmPlot::BeginPlay()
 	IconSetter->HarvestIconScale = HarvestIconScale;
 	PerSecondSystems.Add(IconSetter);
 
+	for (UMassProcessor* Processor : PerSecondSystems)
+	{
+		Processor->CallInitialize(this, EntityManager.AsShared());
+	}
+
 	HarvestIconISMC->SetCullDistances(IconNearCullDistance, IconFarCullDistance);
 
 	// Create ISMCs for each mesh type
