CIS Issue 936433, added missing header inclusion to Mass test types.

#rnx

diff --git a/Engine/Source/Developer/MassEntityTestSuite/Public/MassEntityTestTypes.h b/Engine/Source/Developer/MassEntityTestSuite/Public/MassEntityTestTypes.h
--- a/Engine/Source/Developer/MassEntityTestSuite/Public/MassEntityTestTypes.h
+++ b/Engine/Source/Developer/MassEntityTestSuite/Public/MassEntityTestTypes.h
@@ -3,6 +3,7 @@
 #pragma once
 
 #include "MassProcessor.h"
+#include "MassEntityQuery.h"
 #include "MassEntityTypes.h"
 #include "MassEntityManager.h"
 #include "AITestsCommon.h"
