[AITestSuite] Changed `AITEST_SCOPED_CHECK` to trigger debug break if the expectations are violated and `test.BreakOnTestFail`is `true` (via UE::AITestSuite::ConditionallyBreakOnTestFail).
Updated one test that was turning a blind eye on some known warnings being emitted, polluting the output but not causing errors.

#rnx

diff --git a/Engine/Source/Developer/AITestSuite/Public/TestableEnsures.h b/Engine/Source/Developer/AITestSuite/Public/TestableEnsures.h
--- a/Engine/Source/Developer/AITestSuite/Public/TestableEnsures.h
+++ b/Engine/Source/Developer/AITestSuite/Public/TestableEnsures.h
@@ -20,8 +20,14 @@ namespace UE::AITestSuite
 	ON_SCOPE_EXIT \
 	{ \
 		--UE::AITestSuite::TestsInProgress; \
+		FAutomationTestExecutionInfo ExecutionInfo; \
+		GetTestRunner().GetExecutionInfo(ExecutionInfo); \
+		if (ExecutionInfo.GetWarningTotal() + ExecutionInfo.GetWarningTotal() > 0) \
+		{ \
+			UE::AITestSuite::ConditionallyBreakOnTestFail(); \
+		} \
 	}; \
-	++UE::AITestSuite::TestsInProgress \
+	++UE::AITestSuite::TestsInProgress
 
 #define testableEnsureMsgf(InExpression, InFormat, ... ) \
 	(LIKELY(!!(InExpression)) || ([&]() \
diff --git a/Engine/Source/Developer/MassEntityTestSuite/Private/MassDependencyTest.cpp b/Engine/Source/Developer/MassEntityTestSuite/Private/MassDependencyTest.cpp
--- a/Engine/Source/Developer/MassEntityTestSuite/Private/MassDependencyTest.cpp
+++ b/Engine/Source/Developer/MassEntityTestSuite/Private/MassDependencyTest.cpp
@@ -445,6 +445,8 @@ struct FCircularDependency : FDependencySolverBase
 		{
 			AITEST_SCOPED_CHECK(TEXT("Detected processing dependency cycle"), 1);
 			AITEST_SCOPED_CHECK(TEXT("Encountered processing dependency cycle"), 1);
+			// this one's added since we know we log these, and don't want to trip the detectors
+			AITEST_SCOPED_CHECK(TEXT("group: None"), 4);
 			Solve();
 		}
 		// every subsequent processor is expected to depend only on the previous one since all the processors use exactly the same resources
