[Mass] fixed a typo in FMassEntityManager::FindOrAddArchetypeGroupType function name

#rnx

diff --git a/Engine/Source/Developer/MassEntityTestSuite/Private/MassArchetypeGroupTest.cpp b/Engine/Source/Developer/MassEntityTestSuite/Private/MassArchetypeGroupTest.cpp
--- a/Engine/Source/Developer/MassEntityTestSuite/Private/MassArchetypeGroupTest.cpp
+++ b/Engine/Source/Developer/MassEntityTestSuite/Private/MassArchetypeGroupTest.cpp
@@ -26,7 +26,7 @@ struct FArchetypeGroup_SingleLevelQuery : FEntityTestBase
 	{
 		const int32 NumEntities = 100;
 
-		FArchetypeGroupType GroupType = EntityManager->FinOrAddArchetypeGroupType(TEXT("TestGroup"));
+		FArchetypeGroupType GroupType = EntityManager->FindOrAddArchetypeGroupType(TEXT("TestGroup"));
 
 		TArray<FMassEntityHandle> Entities;
 		TArray<FMassEntityHandle> VerifyEntities;
@@ -115,7 +115,7 @@ struct FArchetypeGroup_MultiLevelQueryBase : FEntityTestBase
 		for (int32 GroupIndex = 0; GroupIndex < GroupsNum; ++GroupIndex)
 		{
 			const FName GroupTypeName(*FString::Printf(TEXT("TestGroup_%d"), GroupIndex));
-			GroupTypes.Add(EntityManager->FinOrAddArchetypeGroupType(GroupTypeName));
+			GroupTypes.Add(EntityManager->FindOrAddArchetypeGroupType(GroupTypeName));
 		}
 		
 		AITEST_NOT_EQUAL(TEXT("Two differently named group types are not expected to be equal"), GroupTypes[0], GroupTypes[1]);
@@ -496,8 +496,8 @@ struct FArchetypeGroup_ApplicationOrder : FEntityTestBase
 	{
 		uint32 ArbitraryNumber = 1677;
 
-		FArchetypeGroupType GroupType1 = EntityManager->FinOrAddArchetypeGroupType(TEXT("TestGroup1"));
-		FArchetypeGroupType GroupType2 = EntityManager->FinOrAddArchetypeGroupType(TEXT("TestGroup2"));
+		FArchetypeGroupType GroupType1 = EntityManager->FindOrAddArchetypeGroupType(TEXT("TestGroup1"));
+		FArchetypeGroupType GroupType2 = EntityManager->FindOrAddArchetypeGroupType(TEXT("TestGroup2"));
 
 		TArray<FMassEntityHandle> Entities;
 		EntityManager->BatchCreateEntities(IntsArchetype, 2, Entities);
@@ -524,8 +524,8 @@ struct FArchetypeGroup_Equivalence : FEntityTestBase
 	{
 		uint32 ArbitraryNumber = 1677;
 
-		FArchetypeGroupType GroupType1 = EntityManager->FinOrAddArchetypeGroupType(TEXT("TestGroup1"));
-		FArchetypeGroupType GroupType2 = EntityManager->FinOrAddArchetypeGroupType(TEXT("TestGroup2"));
+		FArchetypeGroupType GroupType1 = EntityManager->FindOrAddArchetypeGroupType(TEXT("TestGroup1"));
+		FArchetypeGroupType GroupType2 = EntityManager->FindOrAddArchetypeGroupType(TEXT("TestGroup2"));
 
 		TArray<FMassEntityHandle> Entities;
 		EntityManager->BatchCreateEntities(IntsArchetype, 2, Entities);
@@ -555,7 +555,7 @@ struct FArchetypeGroup_NewArchetype : FEntityTestBase
 	{
 		uint32 ArbitraryNumber = 1677;
 
-		FArchetypeGroupType GroupType = EntityManager->FinOrAddArchetypeGroupType(TEXT("TestGroup"));
+		FArchetypeGroupType GroupType = EntityManager->FindOrAddArchetypeGroupType(TEXT("TestGroup"));
 
 		const FMassEntityHandle Entity = EntityManager->CreateEntity(IntsArchetype);
 
@@ -586,7 +586,7 @@ struct FArchetypeGroup_NonGroupChange : FEntityTestBase
 	{
 		uint32 ArbitraryNumber = 1677;
 
-		FArchetypeGroupType GroupType = EntityManager->FinOrAddArchetypeGroupType(TEXT("TestGroup"));
+		FArchetypeGroupType GroupType = EntityManager->FindOrAddArchetypeGroupType(TEXT("TestGroup"));
 		const FArchetypeGroupHandle TargetGroupHandle(GroupType, ArbitraryNumber);
 
 		const FMassEntityHandle Entity = EntityManager->CreateEntity(IntsArchetype);
@@ -613,8 +613,8 @@ struct FArchetypeGroup_Remove : FEntityTestBase
 	virtual bool InstantTest() override
 	{
 		constexpr uint32 ArbitraryNumber = 1677;
-		const FArchetypeGroupType GroupTypeA = EntityManager->FinOrAddArchetypeGroupType(TEXT("TestGroupA"));
-		const FArchetypeGroupType GroupTypeB = EntityManager->FinOrAddArchetypeGroupType(TEXT("TestGroupB"));
+		const FArchetypeGroupType GroupTypeA = EntityManager->FindOrAddArchetypeGroupType(TEXT("TestGroupA"));
+		const FArchetypeGroupType GroupTypeB = EntityManager->FindOrAddArchetypeGroupType(TEXT("TestGroupB"));
 		const FMassArchetypeHandle OriginalArchetype = IntsArchetype;
 
 		TArray<FMassEntityHandle> Entities;
diff --git a/Engine/Source/Runtime/MassEntity/Private/MassEntityManager.cpp b/Engine/Source/Runtime/MassEntity/Private/MassEntityManager.cpp
--- a/Engine/Source/Runtime/MassEntity/Private/MassEntityManager.cpp
+++ b/Engine/Source/Runtime/MassEntity/Private/MassEntityManager.cpp
@@ -963,7 +963,7 @@ UE::Mass::FArchetypeGroupHandle FMassEntityManager::GetGroupForEntity(FMassEntit
 	return UE::Mass::FArchetypeGroupHandle();
 }
 
-UE::Mass::FArchetypeGroupType FMassEntityManager::FinOrAddArchetypeGroupType(const FName GroupName)
+UE::Mass::FArchetypeGroupType FMassEntityManager::FindOrAddArchetypeGroupType(const FName GroupName)
 {
 	LLM_SCOPE_BYNAME(TEXT("Mass/EntityManager"));
 	const int32* FoundGroupIndex = GroupNameToTypeIndex.Find(GroupName);
diff --git a/Engine/Source/Runtime/MassEntity/Public/MassEntityManager.h b/Engine/Source/Runtime/MassEntity/Public/MassEntityManager.h
--- a/Engine/Source/Runtime/MassEntity/Public/MassEntityManager.h
+++ b/Engine/Source/Runtime/MassEntity/Public/MassEntityManager.h
@@ -373,7 +373,7 @@ public:
 	 * Fetches FArchetypeGroupType instance (copy) associated with the given GroupName. A new group type
 	 * is created if GroupName has not been used in the past.
 	 */
-	UE_API UE::Mass::FArchetypeGroupType FinOrAddArchetypeGroupType(const FName GroupName);
+	UE_API UE::Mass::FArchetypeGroupType FindOrAddArchetypeGroupType(const FName GroupName);
 
 	UE_API const UE::Mass::FArchetypeGroups& GetGroupsForArchetype(const FMassArchetypeHandle& ArchetypeHandle) const;
 
