[Mass] Mass Relations
Including:
* MassTypeManager extension for registering relation types.
* Mass Relation Manager and the basic relation types
* added a FMassRelationManager instance as a member to FMassEntityManager, and handling of built-in type registration handling
* one built-in relation type (child of) along with some tests for that and for the generic relations
* added relation-creating capabilities to the Mass EntityBuilder

diff --git a/Engine/Source/Developer/MassEntityTestSuite/Private/MassChildOfRelationTest.cpp b/Engine/Source/Developer/MassEntityTestSuite/Private/MassChildOfRelationTest.cpp
new file mode 100644
--- /dev/null
+++ b/Engine/Source/Developer/MassEntityTestSuite/Private/MassChildOfRelationTest.cpp
@@ -0,0 +1,612 @@
+// Copyright Epic Games, Inc. All Rights Reserved.
+
+#include "AITestsCommon.h"
+#include "MassEntityBuilder.h"
+#include "MassRelationCommands.h"
+#include "MassEntityTestTypes.h"
+#include "MassExecutionContext.h"
+#include "Algo/Compare.h"
+#include "Relations/MassChildOf.h"
+
+#define LOCTEXT_NAMESPACE "MassTest"
+
+UE_DISABLE_OPTIMIZATION_SHIP
+
+namespace UE::Mass::Test::ChildOfRelation
+{
+	using namespace UE::Mass;
+
+	static TNotNull<UScriptStruct*> MakeRelationType(FName TypeName)
+	{
+		UScriptStruct* NewStruct = NewObject<UScriptStruct>(UScriptStruct::StaticClass(), FName(*FString::Printf(TEXT("Relation_%s"), *TypeName.ToString())), RF_Public);
+		NewStruct->SetSuperStruct(FMassRelation::StaticStruct());
+
+		return NewStruct;
+	}
+
+	struct FChildOfBase : FEntityTestBase
+	{
+		enum class EStructure : uint8
+		{
+			String,
+			Tree,
+			MAX
+		};
+		EStructure StructureType = EStructure::MAX;
+
+		using Super = FEntityTestBase;
+		TArray<FMassEntityHandle> CreatedEntities;
+		int32 NumEntities = -1;
+		UMassTestProcessorBase* Processor = nullptr;
+		bool bAPISupportsReparenting = true;
+		FTypeHandle RelationTypeHandle;
+		UE::Mass::FRelationManager* RelationManager = nullptr;
+
+		virtual void BuildHierarchy() = 0;
+		virtual bool ReparentEntity(const int32 ChildIndex, const int32 ParentIndex)
+		{
+			return false;
+		}
+		virtual bool DeleteEntity(const int32 EntityIndex)
+		{
+			EntityManager->DestroyEntity(CreatedEntities[EntityIndex]);
+			return true;
+		}
+
+		virtual void SetUpRelationHandle()
+		{
+			RelationTypeHandle = EntityManager->GetTypeManager().GetRelationTypeHandle(FMassChildOfRelation::StaticStruct());
+			ensure(RelationTypeHandle == UE::Mass::Relations::ChildOfHandle);
+		}
+
+		virtual bool SetUp() override
+		{
+			check(StructureType != EStructure::MAX);
+			switch (StructureType)
+			{
+			case EStructure::String:
+				NumEntities = 5;
+				break;
+			case EStructure::Tree:
+				NumEntities = 8;
+				break;
+			}
+
+			if (Super::SetUp())
+			{
+				SetUpRelationHandle();
+				RelationManager = &EntityManager->GetRelationManager();
+
+				Processor = NewTestProcessor<UMassTestProcessorBase>(EntityManager);
+				Processor->EntityQuery.AddRequirement<FTestFragment_Int>(EMassFragmentAccess::ReadOnly);
+				Processor->EntityQuery.AddRequirement<FMassChildOfFragment>(EMassFragmentAccess::ReadOnly, EMassFragmentPresence::Optional);
+				Processor->EntityQuery.GroupBy(EntityManager->GetTypeManager().GetRelationTypeChecked(RelationTypeHandle).RegisteredGroupType);
+
+				BuildHierarchy();
+
+				return true;
+			}
+			return false;
+		}
+
+		virtual void ExecuteTestProcessor(TArray<int32>& Scratchpad)
+		{
+			Scratchpad.Reset();
+			Scratchpad.AddUninitialized(CreatedEntities.Num());
+			for (int32 EntityIndex = 0; EntityIndex < CreatedEntities.Num(); ++EntityIndex)
+			{
+				Scratchpad[EntityIndex] = -1;
+			}
+
+			EntityManager->FlushCommands();
+
+			Processor->ForEachEntityChunkExecutionFunction = [RelationManager = RelationManager, CreatedEntitiesView = MakeArrayView(CreatedEntities), &Scratchpad](FMassExecutionContext& Context)
+			{
+				// @todo should this be the right way to fetch parents?
+				TConstArrayView<FMassChildOfFragment> Parents = Context.GetFragmentView<FMassChildOfFragment>();
+				if (Parents.Num())
+				{
+					for (auto EntityIt = Context.CreateEntityIterator(); EntityIt; ++EntityIt)
+					{
+						// get parent
+						const FMassEntityHandle ParentHandle = Parents[EntityIt].Parent;
+						if (ParentHandle.IsValid())
+						{
+							const int32 EntityGlobalIndex = CreatedEntitiesView.Find(Context.GetEntity(EntityIt));
+							const int32 ParentGlobalIndex = CreatedEntitiesView.Find(ParentHandle);
+							check(EntityGlobalIndex != INDEX_NONE && ParentGlobalIndex != INDEX_NONE);
+							Scratchpad[EntityGlobalIndex] = Scratchpad[ParentGlobalIndex] * 10 + EntityGlobalIndex;
+						}
+					}
+				}
+				else
+				{
+					// no-parent
+					for (int32 EntityIndex = 0; EntityIndex < Context.GetNumEntities(); ++EntityIndex)
+					{
+						const int32 EntityGlobalIndex = CreatedEntitiesView.Find(Context.GetEntity(EntityIndex));
+						check(EntityGlobalIndex != INDEX_NONE);
+						Scratchpad[EntityGlobalIndex] = EntityGlobalIndex;
+					}
+				}
+			};
+
+			Processor->TestExecute(EntityManager);
+		}
+
+		virtual bool InstantTest() override
+		{
+			TArray<int32> Scratchpad;
+			ExecuteTestProcessor(Scratchpad);
+
+			TArray<int32> ExpectedValues;
+
+			switch (StructureType)
+			{
+			case EStructure::String:
+				// requires SetUp() override to implement the following structure, indicated with indices to CreatedEntities:
+				//	0
+				//	L	1
+				//		L	2
+				//			L	3
+				//				L	4
+				ExpectedValues = { 0, 1, 12, 123, 1234 };
+				break;
+
+			case EStructure::Tree:
+				// requires SetUp() override to implement the following structure, indicated with indices to CreatedEntities:
+				//	0
+				//	+	1
+				//	|	+	3
+				//	|	+	4
+				//	|	L	5
+				//	L	2
+				//		+	6
+				//		L	7
+				ExpectedValues = { 0, 1, 2, 13, 14, 15, 26, 27 };
+				break;
+			}
+			AITEST_TRUE("Processing hierarchy produces expected results", Algo::Compare(Scratchpad, ExpectedValues));
+
+			if (bAPISupportsReparenting == false)
+			{
+				// the API is just for entity creation, we end the test here.
+				return true;
+			}
+
+			// here we reparent a leaf node
+			AITEST_TRUE("Reparenting leaf node.", ReparentEntity(CreatedEntities.Num() - 1, 0));
+			switch (StructureType)
+			{
+			case EStructure::String:
+				// requires SetUp() override to implement the following structure, indicated with indices to CreatedEntities:
+				//	0
+				//	L	1
+				//	|	L	2
+				//	|		L	3
+				//	L	4
+				ExpectedValues = { 0, 1, 12, 123, 4 };
+				break;
+
+			case EStructure::Tree:
+				// requires SetUp() override to implement the following structure, indicated with indices to CreatedEntities:
+				//	0
+				//	+	1
+				//	|	+	3
+				//	|	+	4
+				//	|	L	5
+				//	L	2
+				//	|	L	6
+				//	L	7
+				ExpectedValues = { 0, 1, 2, 13, 14, 15, 26, 7 };
+				break;
+			}
+			ExecuteTestProcessor(Scratchpad);
+			AITEST_TRUE("Reparenting a leaf produces expected results", Algo::Compare(Scratchpad, ExpectedValues));
+
+			AITEST_TRUE("Reparenting a mid-node to a leaf node.", ReparentEntity(2, 4));
+			switch (StructureType)
+			{
+			case EStructure::String:
+				// requires SetUp() override to implement the following structure, indicated with indices to CreatedEntities:
+				//	0
+				//	L	1
+				//	L	4
+				//		L	2
+				//			L	3
+				ExpectedValues = { 0, 1, 42, 423, 4 };
+				break;
+
+			case EStructure::Tree:
+				// requires SetUp() override to implement the following structure, indicated with indices to CreatedEntities:
+				//	0
+				//	+	1
+				//	|	+	3
+				//	|	+	4
+				//	|	|	L	2
+				//	|	|		L	6
+				//	|	L	5
+				//	L	7
+				ExpectedValues = { 0, 1, 142, 13, 14, 15, 1426, 7 };
+				break;
+			}
+			ExecuteTestProcessor(Scratchpad);
+			AITEST_TRUE("Reparenting a subtree to a leaf produces expected results", Algo::Compare(Scratchpad, ExpectedValues));
+
+			AITEST_TRUE("Deleting a leaf node", DeleteEntity(3));
+			switch (StructureType)
+			{
+			case EStructure::String:
+				// requires SetUp() override to implement the following structure, indicated with indices to CreatedEntities:
+				//	0
+				//	L	1
+				//	L	4
+				//		L	2
+				ExpectedValues = { 0, 1, 42, -1, 4 };
+				break;
+
+			case EStructure::Tree:
+				// requires SetUp() override to implement the following structure, indicated with indices to CreatedEntities:
+				//	0
+				//	+	1
+				//	|	+	4
+				//	|	|	L	2
+				//	|	|		L	6
+				//	|	L	5
+				//	L	7
+				ExpectedValues = { 0, 1, 142, -1, 14, 15, 1426, 7 };
+				break;
+			}
+			ExecuteTestProcessor(Scratchpad);
+			AITEST_TRUE("Delete a leaf node produces expected results", Algo::Compare(Scratchpad, ExpectedValues));
+
+			AITEST_TRUE("Deleting mid-level node", DeleteEntity(4));
+			switch (StructureType)
+			{
+			case EStructure::String:
+				// requires SetUp() override to implement the following structure, indicated with indices to CreatedEntities:
+				//	0
+				//	+	1
+				ExpectedValues = { 0, 1, -1, -1, -1 };
+				break;
+
+			case EStructure::Tree:
+				// requires SetUp() override to implement the following structure, indicated with indices to CreatedEntities:
+				//	0
+				//	+	1
+				//	|	L	5
+				//	L	7
+				ExpectedValues = { 0, 1, -1, -1, -1, 15, -1, 7 };
+				break;
+			}
+			ExecuteTestProcessor(Scratchpad);
+			AITEST_TRUE("Delete a mid-level node produces expected results", Algo::Compare(Scratchpad, ExpectedValues));
+
+			AITEST_TRUE("Deleting mid-level node", DeleteEntity(0));
+			switch (StructureType)
+			{
+			case EStructure::String:
+				ExpectedValues = { -1, -1, -1, -1, -1 };
+				break;
+
+			case EStructure::Tree:
+				ExpectedValues = { -1, -1, -1, -1, -1, -1, -1, -1 };
+				break;
+			}
+			ExecuteTestProcessor(Scratchpad);
+			AITEST_TRUE("Delete a top-level node produces expected results", Algo::Compare(Scratchpad, ExpectedValues));
+
+			return true;
+		}
+	};
+
+	struct FChildOf_IndividualAPI_StringHierarchy : FChildOfBase
+	{
+		FChildOf_IndividualAPI_StringHierarchy()
+		{
+			StructureType = EStructure::String;
+		}
+
+		virtual void BuildHierarchy() override
+		{
+			EntityManager->BatchCreateEntities(IntsArchetype, {}, NumEntities, CreatedEntities);
+			// set up indices
+			for (int32 Index = 0; Index < CreatedEntities.Num(); ++Index)
+			{
+				FTestFragment_Int& IndexCounterFragment = EntityManager->GetFragmentDataChecked<FTestFragment_Int>(CreatedEntities[Index]);
+				IndexCounterFragment.Value = Index;
+			}
+
+			for (int32 ChildIndex = 1; ChildIndex < CreatedEntities.Num(); ++ChildIndex)
+			{
+				RelationManager->CreateRelationInstance(RelationTypeHandle, CreatedEntities[ChildIndex], CreatedEntities[ChildIndex - 1]);
+			}
+		}
+
+		virtual bool ReparentEntity(const int32 ChildIndex, const int32 ParentIndex) override
+		{
+			return RelationManager->CreateRelationInstance(RelationTypeHandle, CreatedEntities[ChildIndex], CreatedEntities[ParentIndex]).IsValid();
+		}
+	};
+	IMPLEMENT_AI_INSTANT_TEST(FChildOf_IndividualAPI_StringHierarchy, "System.Mass.Relations.ChildOf.IndividualAPI.StringHierarchy");
+
+	struct FChildOf_IndividualAPI_TreeHierarchy : FChildOfBase
+	{
+		FChildOf_IndividualAPI_TreeHierarchy()
+		{
+			StructureType = EStructure::Tree;
+		}
+
+		virtual void BuildHierarchy() override
+		{
+			EntityManager->BatchCreateEntities(IntsArchetype, {}, NumEntities, CreatedEntities);
+			// set up indices
+			for (int32 Index = 0; Index < CreatedEntities.Num(); ++Index)
+			{
+				FTestFragment_Int& IndexCounterFragment = EntityManager->GetFragmentDataChecked<FTestFragment_Int>(CreatedEntities[Index]);
+				IndexCounterFragment.Value = Index;
+			}
+
+			// [1] - child of [0]
+			RelationManager->CreateRelationInstance(RelationTypeHandle, CreatedEntities[1], CreatedEntities[0]);
+			// [2] - child of [0]
+			RelationManager->CreateRelationInstance(RelationTypeHandle, CreatedEntities[2], CreatedEntities[0]);
+			// [3] - child of [1]
+			RelationManager->CreateRelationInstance(RelationTypeHandle, CreatedEntities[3], CreatedEntities[1]);
+			// [4] - child of [1]
+			RelationManager->CreateRelationInstance(RelationTypeHandle, CreatedEntities[4], CreatedEntities[1]);
+			// [5] - child of [1]
+			RelationManager->CreateRelationInstance(RelationTypeHandle, CreatedEntities[5], CreatedEntities[1]);
+			// [6] - child of [2]
+			RelationManager->CreateRelationInstance(RelationTypeHandle, CreatedEntities[6], CreatedEntities[2]);
+			// [7] - child of [2]
+			RelationManager->CreateRelationInstance(RelationTypeHandle, CreatedEntities[7], CreatedEntities[2]);
+		};
+
+		virtual bool ReparentEntity(const int32 ChildIndex, const int32 ParentIndex) override
+		{
+			return RelationManager->CreateRelationInstance(RelationTypeHandle, CreatedEntities[ChildIndex], CreatedEntities[ParentIndex]).IsValid();
+		}
+	};
+	IMPLEMENT_AI_INSTANT_TEST(FChildOf_IndividualAPI_TreeHierarchy, "System.Mass.Relations.ChildOf.IndividualAPI.TreeHierarchy");
+
+	struct FChildOfEntityBuilder_StringHierarchy : FChildOfBase
+	{
+		FChildOfEntityBuilder_StringHierarchy()
+		{
+			StructureType = EStructure::String;
+			bAPISupportsReparenting = false;
+		}
+
+		virtual void BuildHierarchy() override
+		{
+			FEntityBuilder Builder = EntityManager->MakeEntityBuilder();
+
+			FTestFragment_Int& IndexCounterFragment = Builder.Add_GetRef<FTestFragment_Int>(CreatedEntities.Num());
+			CreatedEntities.Add(Builder.CommitAndReprepare());
+
+			for (int32 ChildIndex = 0; ChildIndex < 4; ++ChildIndex)
+			{
+				Builder.AddRelation(RelationTypeHandle, CreatedEntities.Last());
+				IndexCounterFragment.Value = CreatedEntities.Num();
+				CreatedEntities.Add(Builder.CommitAndReprepare());
+			}
+		}
+	};
+	IMPLEMENT_AI_INSTANT_TEST(FChildOfEntityBuilder_StringHierarchy, "System.Mass.Relations.ChildOf.EntityBuilder.StringHierarchy");
+
+	struct FChildOfEntityBuilder_TreeHierarchy : FChildOfBase
+	{
+		FChildOfEntityBuilder_TreeHierarchy()
+		{
+			StructureType = EStructure::Tree;
+			bAPISupportsReparenting = false;
+		}
+
+		virtual void BuildHierarchy() override
+		{
+			FEntityBuilder Builder = EntityManager->MakeEntityBuilder();
+
+			FTestFragment_Int& IndexCounterFragment = Builder.Add_GetRef<FTestFragment_Int>(CreatedEntities.Num());
+			// [0] - parent entity
+			CreatedEntities.Add(Builder.CommitAndReprepare());
+
+			Builder.AddRelation(RelationTypeHandle, CreatedEntities.Last());
+			// [1] - child of [0]
+			IndexCounterFragment.Value = CreatedEntities.Num();
+			CreatedEntities.Add(Builder.CommitAndReprepare());
+			// [2] - child of [0]
+			IndexCounterFragment.Value = CreatedEntities.Num();
+			CreatedEntities.Add(Builder.CommitAndReprepare());
+
+			Builder.AddRelation(RelationTypeHandle, CreatedEntities[1]);
+			// [3] - child of [1]
+			IndexCounterFragment.Value = CreatedEntities.Num();
+			CreatedEntities.Add(Builder.CommitAndReprepare());
+			// [4] - child of [1]
+			IndexCounterFragment.Value = CreatedEntities.Num();
+			CreatedEntities.Add(Builder.CommitAndReprepare());
+			// [5] - child of [1]
+			IndexCounterFragment.Value = CreatedEntities.Num();
+			CreatedEntities.Add(Builder.CommitAndReprepare());
+
+			Builder.AddRelation(RelationTypeHandle, CreatedEntities[2]);
+			// [6] - child of [2]
+			IndexCounterFragment.Value = CreatedEntities.Num();
+			CreatedEntities.Add(Builder.CommitAndReprepare());
+			// [7] - child of [2]
+			IndexCounterFragment.Value = CreatedEntities.Num();
+			CreatedEntities.Add(Builder.CommitAndReprepare());
+		}
+	};
+	IMPLEMENT_AI_INSTANT_TEST(FChildOfEntityBuilder_TreeHierarchy, "System.Mass.Relations.ChildOf.EntityBuilder.TreeHierarchy");
+
+	struct FChildOfBatchAPI_Tree : FChildOfBase
+	{
+		FChildOfBatchAPI_Tree()
+		{
+			StructureType = EStructure::Tree;
+		}
+
+		virtual void BuildHierarchy() override
+		{
+			EntityManager->BatchCreateEntities(IntsArchetype, {}, NumEntities, CreatedEntities);
+			// set up indices
+			for (int32 Index = 0; Index < CreatedEntities.Num(); ++Index)
+			{
+				FTestFragment_Int& IndexCounterFragment = EntityManager->GetFragmentDataChecked<FTestFragment_Int>(CreatedEntities[Index]);
+				IndexCounterFragment.Value = Index;
+			}
+
+			// [1, 2] - child of [0]
+			EntityManager->BatchCreateRelations(RelationTypeHandle, MakeArrayView(&CreatedEntities[1], 1), MakeArrayView(&CreatedEntities[0], 1));
+			EntityManager->BatchCreateRelations(RelationTypeHandle, MakeArrayView(&CreatedEntities[2], 1), MakeArrayView(&CreatedEntities[0], 1));
+			// [3, 4, 5] - child of [1]
+			EntityManager->BatchCreateRelations(RelationTypeHandle, MakeArrayView(&CreatedEntities[3], 1), MakeArrayView(&CreatedEntities[1], 1));
+			EntityManager->BatchCreateRelations(RelationTypeHandle, MakeArrayView(&CreatedEntities[4], 1), MakeArrayView(&CreatedEntities[1], 1));
+			EntityManager->BatchCreateRelations(RelationTypeHandle, MakeArrayView(&CreatedEntities[5], 1), MakeArrayView(&CreatedEntities[1], 1));
+			// [6, 7] - child of [2]
+			EntityManager->BatchCreateRelations(RelationTypeHandle, MakeArrayView(&CreatedEntities[6], 1), MakeArrayView(&CreatedEntities[2], 1));
+			EntityManager->BatchCreateRelations(RelationTypeHandle, MakeArrayView(&CreatedEntities[7], 1), MakeArrayView(&CreatedEntities[2], 1));
+		}
+
+		virtual bool ReparentEntity(const int32 ChildIndex, const int32 ParentIndex) override
+		{
+			return EntityManager->BatchCreateRelations(RelationTypeHandle, MakeArrayView(&CreatedEntities[ChildIndex], 1), MakeArrayView(&CreatedEntities[ParentIndex], 1));
+		}
+	};
+	IMPLEMENT_AI_INSTANT_TEST(FChildOfBatchAPI_Tree, "System.Mass.Relations.ChildOf.BatchAPI");
+
+	struct FChildOfCommands_StringHierarchy : FChildOfBase
+	{
+		FChildOfCommands_StringHierarchy()
+		{
+			StructureType = EStructure::String;
+		}
+
+		virtual void BuildHierarchy() override
+		{
+			EntityManager->BatchCreateEntities(IntsArchetype, {}, NumEntities, CreatedEntities);
+			// set up indices
+			for (int32 Index = 0; Index < CreatedEntities.Num(); ++Index)
+			{
+				FTestFragment_Int& IndexCounterFragment = EntityManager->GetFragmentDataChecked<FTestFragment_Int>(CreatedEntities[Index]);
+				IndexCounterFragment.Value = Index;
+				if (Index > 0)
+				{
+					// issue relationship-creating commands, done here for convenience 
+					EntityManager->Defer().PushCommand<FMassCommandMakeRelation<FMassChildOfRelation>>(CreatedEntities[Index], CreatedEntities[Index-1]);
+				}
+			}
+
+			EntityManager->FlushCommands();
+		}
+
+		virtual bool ReparentEntity(const int32 ChildIndex, const int32 ParentIndex) override
+		{
+			EntityManager->Defer().PushCommand<FMassCommandMakeRelation<FMassChildOfRelation>>(CreatedEntities[ChildIndex], CreatedEntities[ParentIndex]);
+			return true;
+		}
+
+		virtual bool DeleteEntity(const int32 EntityIndex) override
+		{
+			EntityManager->Defer().DestroyEntity(CreatedEntities[EntityIndex]);
+			return true;
+		}
+	};
+	IMPLEMENT_AI_INSTANT_TEST(FChildOfCommands_StringHierarchy, "System.Mass.Relations.ChildOf.Commands.StringHierarchy");
+
+	struct FChildOfCommands_TreeHierarchy : FChildOfBase
+	{
+		FChildOfCommands_TreeHierarchy()
+		{
+			StructureType = EStructure::Tree;
+		}
+
+		virtual void BuildHierarchy() override
+		{
+			EntityManager->BatchCreateEntities(IntsArchetype, {}, NumEntities, CreatedEntities);
+			
+			for (int32 Index = 0; Index < CreatedEntities.Num(); ++Index)
+			{
+				FTestFragment_Int& IndexCounterFragment = EntityManager->GetFragmentDataChecked<FTestFragment_Int>(CreatedEntities[Index]);
+				IndexCounterFragment.Value = Index;
+			}
+
+			// [1, 2] - child of [0]
+			EntityManager->Defer().PushCommand<FMassCommandMakeRelation<FMassChildOfRelation>>(MakeArrayView(&CreatedEntities[1], 2), MakeArrayView(&CreatedEntities[0], 1));
+			// [3, 4, 5] - child of [1]
+			EntityManager->Defer().PushCommand<FMassCommandMakeRelation<FMassChildOfRelation>>(MakeArrayView(&CreatedEntities[3], 3), MakeArrayView(&CreatedEntities[1], 1));
+			// [6, 7] - child of [2]
+			EntityManager->Defer().PushCommand<FMassCommandMakeRelation<FMassChildOfRelation>>(MakeArrayView(&CreatedEntities[6], 2), MakeArrayView(&CreatedEntities[2], 1));
+
+			EntityManager->FlushCommands();
+		}
+
+		virtual bool ReparentEntity(const int32 ChildIndex, const int32 ParentIndex) override
+		{
+			EntityManager->Defer().PushCommand<FMassCommandMakeRelation<FMassChildOfRelation>>(CreatedEntities[ChildIndex], CreatedEntities[ParentIndex]);
+			EntityManager->FlushCommands();
+			return true;
+		}
+
+		virtual bool DeleteEntity(const int32 EntityIndex) override
+		{
+			EntityManager->Defer().DestroyEntity(CreatedEntities[EntityIndex]);
+			return true;
+		}
+	};
+	IMPLEMENT_AI_INSTANT_TEST(FChildOfCommands_TreeHierarchy, "System.Mass.Relations.ChildOf.Commands.TreeHierarchy");
+
+	struct FChildOfBatchAPI_Tree_NoMapping : FChildOfBatchAPI_Tree
+	{
+		virtual void SetUpRelationHandle() override
+		{
+			UScriptStruct* MaplessRelationType = MakeRelationType("ChildOfNoMapping");
+			FRelationTypeTraits RelationTraits(EntityManager->GetTypeManager().GetRelationTypeChecked(UE::Mass::Relations::ChildOfHandle), MaplessRelationType);
+			RelationTraits.RoleTraits[static_cast<uint8>(ERelationRole::Subject)].RequiresExternalMapping = EExternalMappingRequired::No;
+
+			CA_ASSUME(MaplessRelationType);
+			RelationTypeHandle = EntityManager->GetTypeManager().RegisterType(MoveTemp(RelationTraits));
+		}
+	};
+	// commented out on purpose, functionality not implemented yet
+	//IMPLEMENT_AI_INSTANT_TEST(FChildOfBatchAPI_Tree_NoMapping, "System.Mass.Relations.ChildOf.NoMapping.BatchAPI");
+
+	// - other tests ---
+	struct FSetChildAsParent : FEntityTestBase
+	{
+		virtual bool InstantTest() override
+		{
+			FTypeHandle RelationTypeHandle = EntityManager->GetTypeManager().GetRelationTypeHandle(FMassChildOfRelation::StaticStruct());
+
+			TArray<FMassEntityHandle> CreatedEntities;
+			EntityManager->BatchCreateEntities(IntsArchetype, {}, 2, CreatedEntities);
+
+			UE::Mass::FRelationManager& RelationManager = EntityManager->GetRelationManager();
+			RelationManager.CreateRelationInstance(RelationTypeHandle, CreatedEntities[0], CreatedEntities[1]);
+			EntityManager->FlushCommands();
+
+			// original parent becomes the child and vice versa:
+			RelationManager.CreateRelationInstance(RelationTypeHandle, CreatedEntities[1], CreatedEntities[0]);
+			EntityManager->FlushCommands();
+
+			TArray<FMassEntityHandle> ParentSubjects = RelationManager.GetRelationSubjects(RelationTypeHandle, CreatedEntities[0]);
+			TArray<FMassEntityHandle> ParentObjects = RelationManager.GetRelationSubjects(RelationTypeHandle, CreatedEntities[0]);
+			TArray<FMassEntityHandle> ChildSubjects = RelationManager.GetRelationSubjects(RelationTypeHandle, CreatedEntities[1]);
+			TArray<FMassEntityHandle> ChildObjects = RelationManager.GetRelationSubjects(RelationTypeHandle, CreatedEntities[1]);
+
+			//AITEST_TRUE()
+
+			return true;
+		}
+	};
+	IMPLEMENT_AI_INSTANT_TEST(FSetChildAsParent, "System.Mass.Relations.ChildOf.SetChildAsParent");
+
+
+} // UE::Mass::Test::ChildOfRelation
+
+UE_ENABLE_OPTIMIZATION_SHIP
+
+#undef LOCTEXT_NAMESPACE
+#undef WITH_BUILDER
\ No newline at end of file
diff --git a/Engine/Source/Developer/MassEntityTestSuite/Private/MassRelationTest.cpp b/Engine/Source/Developer/MassEntityTestSuite/Private/MassRelationTest.cpp
new file mode 100644
--- /dev/null
+++ b/Engine/Source/Developer/MassEntityTestSuite/Private/MassRelationTest.cpp
@@ -0,0 +1,706 @@
+// Copyright Epic Games, Inc. All Rights Reserved.
+
+#include "AITestsCommon.h"
+#include "MassRelationManager.h"
+#include "MassEntityTestTypes.h"
+
+#define LOCTEXT_NAMESPACE "MassTest"
+
+UE_DISABLE_OPTIMIZATION_SHIP
+
+namespace UE::Mass::Test::Relations
+{
+	using namespace UE::Mass;
+
+	struct FRelationTestBase : FEntityTestBase
+	{
+		UScriptStruct* MyRelationType = nullptr;
+		const FRelationTypeTraits* RelationTraits = nullptr;
+		FTypeHandle RelationTypeHandle;
+		FRelationManager* RelationshipManagerPtr = nullptr;
+
+		static TNotNull<UScriptStruct*> MakeRelationType(FName TypeName)
+		{
+			UScriptStruct* NewStruct = NewObject<UScriptStruct>(UScriptStruct::StaticClass(), FName(*FString::Printf(TEXT("Relation_%s"), *TypeName.ToString())), RF_Public);
+			NewStruct->SetSuperStruct(FMassRelation::StaticStruct());
+
+			return NewStruct;
+		}
+
+		template<UE::Mass::CElement TBase>
+		static TNotNull<UScriptStruct*> MakeElementType(FName TypeName)
+		{	
+			UScriptStruct* NewStruct = NewObject<UScriptStruct>(UScriptStruct::StaticClass(), FName(*FString::Printf(TEXT("%s"), *TypeName.ToString())), RF_Public);
+			NewStruct->SetSuperStruct(TBase::StaticStruct());
+			return NewStruct;
+		}
+
+		static FRelationTypeTraits CreateRelationTypeDescription(const FName FragmentName)
+		{
+			UScriptStruct* FragmentType = MakeRelationType(FragmentName);
+			return CreateRelationTypeDescription(FragmentType);
+		}
+
+		static FRelationTypeTraits CreateRelationTypeDescription(TNotNull<UScriptStruct*> RelationType)
+		{
+			FRelationTypeTraits RelationTraits(RelationType);
+			RelationTraits.RoleTraits[static_cast<uint8>(ERelationRole::Object)].DestructionPolicy = ERemovalPolicy::Destroy;
+			RelationTraits.RoleTraits[static_cast<uint8>(ERelationRole::Object)].bExclusive = true;
+			RelationTraits.RoleTraits[static_cast<uint8>(ERelationRole::Subject)].DestructionPolicy = ERemovalPolicy::CleanUp;
+			RelationTraits.RoleTraits[static_cast<uint8>(ERelationRole::Subject)].bExclusive = false;
+			RelationTraits.bHierarchical = true;
+			return RelationTraits;
+		}
+
+		virtual bool SetUp() override
+		{
+			if (FEntityTestBase::SetUp())
+			{
+				MyRelationType = MakeRelationType("MyChildOfRelation");
+				FRelationTypeTraits TempRelationTraits = CreateRelationTypeDescription("MyChildOfRelation");
+				TweakRelation(TempRelationTraits);
+				RelationTypeHandle = EntityManager->GetTypeManager().RegisterType(MoveTemp(TempRelationTraits));
+				RelationTraits = &EntityManager->GetTypeManager().GetRelationTypeChecked(RelationTypeHandle);
+				RelationshipManagerPtr = &EntityManager->GetRelationManager();
+			}
+			return RelationTypeHandle.IsValid();
+		}
+
+		virtual void TweakRelation(FRelationTypeTraits& TempRelationTraits)
+		{
+		}
+
+		FMassEntityHandle GetRelationSubject(const FMassEntityHandle ObjectEntity) const 
+		{
+			check(RelationshipManagerPtr);
+			TArray<FMassEntityHandle> Subjects = RelationshipManagerPtr->GetRelationSubjects(RelationTypeHandle, ObjectEntity);
+			return Subjects.Num() ? Subjects[0] : FMassEntityHandle();
+		}
+
+		FMassEntityHandle GetRelationObject(const FMassEntityHandle SubjectEntity) const
+		{
+			check(RelationshipManagerPtr);
+			TArray<FMassEntityHandle> Objects = RelationshipManagerPtr->GetRelationObjects(RelationTypeHandle, SubjectEntity);
+			return Objects.Num() ? Objects[0] : FMassEntityHandle();
+		}
+	};
+
+	struct FSingleRelationCreation : FRelationTestBase
+	{
+		FMassArchetypeHandle OriginalArchetype;
+		TArray<FMassEntityHandle> CreatedEntities;
+
+		virtual bool SetUp() override
+		{
+			FRelationTestBase::SetUp();
+			OriginalArchetype = IntsArchetype;
+			EntityManager->BatchCreateEntities(OriginalArchetype, 2, CreatedEntities);
+			return true;
+		}
+
+		virtual bool InstantTest() override
+		{
+			const FMassEntityHandle ChildEntity = CreatedEntities[0];
+			const FMassEntityHandle ParentEntity = CreatedEntities[1];
+
+			FMassEntityHandle CreatedRelationEntity = RelationshipManagerPtr->CreateRelationInstance(RelationTypeHandle, ChildEntity, ParentEntity);
+			EntityManager->FlushCommands();
+
+			AITEST_TRUE("Valid relation entity has been created", CreatedRelationEntity.IsValid());
+
+			TArray<FMassEntityHandle> TestedEntities = RelationshipManagerPtr->GetRelationSubjects(RelationTypeHandle, ChildEntity);
+			AITEST_TRUE("No subjects point to the ChildEntity as the relation's object", TestedEntities.IsEmpty());
+			TestedEntities = RelationshipManagerPtr->GetRelationSubjects(RelationTypeHandle, ParentEntity);
+			AITEST_TRUE("ParentEntity is an object of a relation of the given type", TestedEntities.Num() > 0);
+			AITEST_TRUE("ParentEntity is an object of exactly one relation of the given type", TestedEntities.Num() == 1);
+			AITEST_TRUE("The ChildEntity is the subject of the give relation, where ParentEntity is the object", TestedEntities[0] == ChildEntity);
+
+			TestedEntities = RelationshipManagerPtr->GetRelationObjects(RelationTypeHandle, ChildEntity);
+			AITEST_TRUE("The ChildEntity is has exactly one object for the given relation type", TestedEntities.Num() == 1);
+			AITEST_TRUE("The object for the given relation for the ChildEntity is the ParentEntity", TestedEntities[0] == ParentEntity);
+			TestedEntities = RelationshipManagerPtr->GetRelationObjects(RelationTypeHandle, ParentEntity);
+			AITEST_TRUE("The ParentEntity has no objects for the given relation type", TestedEntities.IsEmpty());
+
+			RelationshipManagerPtr->DestroyRelationInstance(RelationTypeHandle, ChildEntity, ParentEntity);
+			EntityManager->FlushCommands();
+			TestedEntities = RelationshipManagerPtr->GetRelationSubjects(RelationTypeHandle, ChildEntity);
+			AITEST_TRUE("No relation subjects for ChildEntity still", TestedEntities.IsEmpty());
+			TestedEntities = RelationshipManagerPtr->GetRelationSubjects(RelationTypeHandle, ParentEntity);
+			AITEST_TRUE("No relation subjects for ParentEntity anymore", TestedEntities.IsEmpty());
+
+			TestedEntities = RelationshipManagerPtr->GetRelationObjects(RelationTypeHandle, ChildEntity);
+			AITEST_TRUE("No relation objects for ChildEntity anymore", TestedEntities.IsEmpty());
+			TestedEntities = RelationshipManagerPtr->GetRelationObjects(RelationTypeHandle, ParentEntity);
+			AITEST_TRUE("No relation objects for ParentEntity still", TestedEntities.IsEmpty());
+
+			return true;
+		}
+
+	};
+	IMPLEMENT_AI_INSTANT_TEST(FSingleRelationCreation, "System.Mass.Relations.CreateSingle");
+
+
+	struct FUnitHierarchy : FRelationTestBase
+	{
+		virtual bool InstantTest() override
+		{
+			const FMassArchetypeHandle OriginalArchetype = IntsArchetype;
+
+			TArray<FMassEntityHandle> Entities;
+			EntityManager->BatchCreateEntities(OriginalArchetype, 2, Entities);
+
+			constexpr int32 ParentIndex = 1;
+			constexpr int32 ChildIndex = 0;
+			RelationshipManagerPtr->CreateRelationInstance(RelationTypeHandle, Entities[ChildIndex], Entities[ParentIndex]);
+
+			const FMassArchetypeHandle ArchetypeParent = EntityManager->GetArchetypeForEntity(Entities[ParentIndex]);
+			const FMassArchetypeHandle ArchetypeChild = EntityManager->GetArchetypeForEntity(Entities[ChildIndex]);
+
+			AITEST_NOT_EQUAL("Child and Parent archetypes", ArchetypeParent, ArchetypeChild);
+			AITEST_NOT_EQUAL("Child and the original archetype", ArchetypeChild, OriginalArchetype);
+			AITEST_NOT_EQUAL("Parent and the original archetype", ArchetypeParent, OriginalArchetype);
+
+			FArchetypeGroupHandle GroupHandle = EntityManager->GetGroupForEntity(Entities[ChildIndex], RelationTraits->RegisteredGroupType);
+			AITEST_EQUAL("Leaf entity's is at expected group level", GroupHandle.GetGroupID(), 1);
+			
+			return true;
+		}
+	};
+	IMPLEMENT_AI_INSTANT_TEST(FUnitHierarchy, "System.Mass.Relations.Hierarchy.Unit");
+
+	struct FChainHierarchy : FRelationTestBase
+	{
+		virtual bool InstantTest() override
+		{
+			constexpr int32 NumEntities = 3;
+			const FMassArchetypeHandle OriginalArchetype = IntsArchetype;
+
+			TArray<FMassEntityHandle> Entities;
+			EntityManager->BatchCreateEntities(OriginalArchetype, NumEntities, Entities);
+
+			constexpr int32 LeafIndex = 1;
+			constexpr int32 MiddleIndex = 0;
+			constexpr int32 RootIndex = 2;
+			RelationshipManagerPtr->CreateRelationInstance(RelationTypeHandle, Entities[LeafIndex], Entities[MiddleIndex]);
+			RelationshipManagerPtr->CreateRelationInstance(RelationTypeHandle, Entities[MiddleIndex], Entities[RootIndex]);
+
+			FMassArchetypeHandle Archetypes[NumEntities] = {
+				EntityManager->GetArchetypeForEntity(Entities[0])
+				, EntityManager->GetArchetypeForEntity(Entities[1])
+				, EntityManager->GetArchetypeForEntity(Entities[2])
+			};
+
+			for (int32 Index = 0; Index < NumEntities; ++Index)
+			{
+				AITEST_NOT_EQUAL("Each pair of archetypes", Archetypes[Index], Archetypes[(Index + 1) % NumEntities]);
+			}
+			{
+				constexpr int32 EntityIndex = LeafIndex;
+				const FMassEntityHandle Object = GetRelationObject(Entities[EntityIndex]);
+				
+				AITEST_EQUAL("Leaf entity's Object vs Middle", Object, Entities[MiddleIndex]);
+				TArray<FMassEntityHandle> Subjects = RelationshipManagerPtr->GetRelationSubjects(RelationTypeHandle, Entities[EntityIndex]);
+				AITEST_TRUE("Leaf entity's has no children", Subjects.IsEmpty());
+
+				FArchetypeGroupHandle GroupHandle = EntityManager->GetGroupForEntity(Entities[EntityIndex], RelationTraits->RegisteredGroupType);
+				AITEST_EQUAL("Leaf entity's is at expected group level", GroupHandle.GetGroupID(), 2);
+			}
+			{
+				constexpr int32 EntityIndex = MiddleIndex;
+				const FMassEntityHandle Object = GetRelationObject(Entities[EntityIndex]);
+				AITEST_EQUAL("Middle entity's Object vs Root", Object, Entities[RootIndex]);
+				TArray<FMassEntityHandle> Subjects = RelationshipManagerPtr->GetRelationSubjects(RelationTypeHandle,  Entities[EntityIndex]);
+				AITEST_TRUE("Middle entity's has a single child", Subjects.Num() == 1);
+				AITEST_TRUE("Middle entity's sole child is the leaf entity", Subjects.Last() == Entities[LeafIndex]);
+
+				FArchetypeGroupHandle GroupHandle = EntityManager->GetGroupForEntity(Entities[EntityIndex], RelationTraits->RegisteredGroupType);
+				AITEST_EQUAL("Leaf entity's is at expected group level", GroupHandle.GetGroupID(), 1);
+			}
+			{
+				constexpr int32 EntityIndex = RootIndex;
+				const FMassEntityHandle Object = GetRelationObject(Entities[EntityIndex]);
+				AITEST_TRUE("Root entity has no parent", Object.IsValid() == false);
+				TArray<FMassEntityHandle> Subjects = RelationshipManagerPtr->GetRelationSubjects(RelationTypeHandle,  Entities[EntityIndex]);
+				AITEST_TRUE("Root entity's has a single child", Subjects.Num() == 1);
+				AITEST_TRUE("Root entity's sole child is the middle entity", Subjects.Last() == Entities[MiddleIndex]);
+
+				FArchetypeGroupHandle GroupHandle = EntityManager->GetGroupForEntity(Entities[EntityIndex], RelationTraits->RegisteredGroupType);
+				AITEST_EQUAL("Leaf entity's is at expected group level", GroupHandle.GetGroupID(), 0);
+			}
+			
+			return true;
+		}
+	};
+	IMPLEMENT_AI_INSTANT_TEST(FChainHierarchy, "System.Mass.Relations.Hierarchy.Chain");
+
+	struct FHierarchyTrivialReparenting : FRelationTestBase
+	{
+		virtual bool InstantTest() override
+		{
+			constexpr int32 NumEntities = 3;
+			const FMassArchetypeHandle OriginalArchetype = IntsArchetype;
+
+			TArray<FMassEntityHandle> Entities;
+			EntityManager->BatchCreateEntities(OriginalArchetype, NumEntities, Entities);
+
+			constexpr int32 LeafIndex = 1;
+			constexpr int32 OriginalParentIndex = 0;
+			constexpr int32 FinalParentIndex = 2;
+			RelationshipManagerPtr->CreateRelationInstance(RelationTypeHandle, Entities[LeafIndex], Entities[OriginalParentIndex]);
+			RelationshipManagerPtr->CreateRelationInstance(RelationTypeHandle, Entities[LeafIndex], Entities[FinalParentIndex]);
+			{
+				constexpr int32 EntityIndex = LeafIndex;
+				const FMassEntityHandle Object = GetRelationObject(Entities[EntityIndex]);
+				AITEST_NOT_EQUAL("Leaf entity's Object vs Original Parent", Object, Entities[OriginalParentIndex]);
+				AITEST_EQUAL("Leaf entity's Object vs Final Parent", Object, Entities[FinalParentIndex]);
+
+				TArray<FMassEntityHandle> Subjects = RelationshipManagerPtr->GetRelationSubjects(RelationTypeHandle,  Entities[EntityIndex]);
+				AITEST_TRUE("Leaf entity's has no children", Subjects.IsEmpty());
+			}
+			{
+				constexpr int32 EntityIndex = OriginalParentIndex;
+				const FMassEntityHandle Object = GetRelationObject(Entities[EntityIndex]);
+				AITEST_FALSE("Original Parent has a parent", Object.IsValid());
+
+				TArray<FMassEntityHandle> Subjects = RelationshipManagerPtr->GetRelationSubjects(RelationTypeHandle,  Entities[EntityIndex]);
+				AITEST_TRUE("Original Parent has no children", Subjects.IsEmpty());
+			}
+			{
+				constexpr int32 EntityIndex = FinalParentIndex;
+				const FMassEntityHandle Object = GetRelationObject(Entities[EntityIndex]);
+				AITEST_FALSE("Final Parent has a parent", Object.IsValid());
+
+				TArray<FMassEntityHandle> Subjects = RelationshipManagerPtr->GetRelationSubjects(RelationTypeHandle,  Entities[EntityIndex]);
+				AITEST_TRUE("Final Parent has one child", Subjects.Num() == 1);
+				AITEST_TRUE("Final Parent's sole child is the Leaf entity", Subjects.Last() == Entities[LeafIndex]);
+			}
+
+			return true;
+		}
+	};
+	IMPLEMENT_AI_INSTANT_TEST(FHierarchyTrivialReparenting, "System.Mass.Relations.Hierarchy.TrivialReparenting");
+
+	struct FComplexHierarchy : FRelationTestBase
+	{
+		using Super = FRelationTestBase;
+
+		int32 NoRelation = INDEX_NONE;
+		int32 SingleChildParent = INDEX_NONE;
+		int32 SingleChildChild = INDEX_NONE;
+		static constexpr int32 NumEntities = 9;
+		bool Result[NumEntities][NumEntities] = {
+				{ false, false, false, false, false, false, false, false, false},
+				{ false, false, false, false, false, false, false, false, false},
+				{ false, false, false, false, false, false, false, false, false},
+				{ false, false, false, false, false, false, false, false, false},
+				{ false, false, false, false, false, false, false, false, false},
+				{ false, false, false, false, false, false, false, false, false},
+				{ false, false, false, false, false, false, false, false, false},
+				{ false, false, false, false, false, false, false, false, false},
+				{ false, false, false, false, false, false, false, false, false},
+			};
+
+		FMassArchetypeHandle OriginalArchetype;
+		TArray<FMassEntityHandle> Entities;
+
+		virtual bool SetUp() override
+		{
+			Super::SetUp();
+
+			OriginalArchetype = IntsArchetype;
+			EntityManager->BatchCreateEntities(OriginalArchetype, NumEntities, Entities);
+
+			RelationshipManagerPtr->CreateRelationInstance(RelationTypeHandle, Entities[3], Entities[0]);
+			Result[3][0] = true;
+			RelationshipManagerPtr->CreateRelationInstance(RelationTypeHandle, Entities[5], Entities[3]);
+			Result[5][3] = true;
+			Result[5][0] = true; // grandparent
+			RelationshipManagerPtr->CreateRelationInstance(RelationTypeHandle, Entities[7], Entities[3]);
+			Result[7][3] = true;
+			Result[7][0] = true; // grandparent
+			RelationshipManagerPtr->CreateRelationInstance(RelationTypeHandle, Entities[6], Entities[5]);
+			Result[6][5] = true;
+			Result[6][3] = true; // grandparent
+			Result[6][0] = true; // grand-grandparent
+
+			RelationshipManagerPtr->CreateRelationInstance(RelationTypeHandle, Entities[8], Entities[0]);
+			Result[8][0] = true;
+
+			RelationshipManagerPtr->CreateRelationInstance(RelationTypeHandle, Entities[4], Entities[1]);
+			Result[4][1] = true;
+
+			NoRelation = 2;
+			SingleChildParent = 1;
+			SingleChildChild = 4;
+			
+
+			// we now have the following hierarchy, indexing Entities array
+			// 0				1				2 - Not in hierarchy
+			// +---	3			+- 4
+			// |	+ 	5
+			// |	|	+	 6
+			// |	|
+			// |	+ 	7
+			// |
+			// +---	8
+			//
+
+			return true;
+		}
+	};
+
+	struct FChildOf : FComplexHierarchy
+	{
+		virtual bool InstantTest() override
+		{
+			// no relation
+			for (int32 EntityIndex = 0; EntityIndex < Entities.Num(); ++EntityIndex)
+			{
+				if (EntityIndex != NoRelation)
+				{
+					AITEST_FALSE("NoRelation is a child of another entity", RelationshipManagerPtr->IsSubjectOfRelation(RelationTypeHandle, Entities[NoRelation], Entities[EntityIndex]));
+					AITEST_FALSE("NoRelation is a parent to another entity", RelationshipManagerPtr->IsSubjectOfRelation(RelationTypeHandle, Entities[EntityIndex], Entities[NoRelation]));
+				}
+			}
+
+			// simple 1-1 hierarchy
+			AITEST_FALSE("(NOT) SingleChildParent is a child of SingleChildChild", RelationshipManagerPtr->IsSubjectOfRelation(RelationTypeHandle, Entities[SingleChildParent], Entities[SingleChildChild]));
+			AITEST_TRUE("SingleChildChild is a child of SingleChildParent", RelationshipManagerPtr->IsSubjectOfRelation(RelationTypeHandle, Entities[SingleChildChild], Entities[SingleChildParent]));
+
+			// We'll test everything else in bulk. Here's what we want to see, rows are "children", columns the result of IsSubjectOfRelation
+			// note that Result contains cumulative results, so we test indirect relations too
+			for (int32 ChildIndex = 0; ChildIndex < Entities.Num(); ++ChildIndex)
+			{
+				for (int32 ParentIndex = 0; ParentIndex < Entities.Num(); ++ParentIndex)
+				{
+					AITEST_EQUAL("Child is a child", RelationshipManagerPtr->IsSubjectOfRelationRecursive(RelationTypeHandle, Entities[ChildIndex], Entities[ParentIndex]), Result[ChildIndex][ParentIndex]);
+				}
+			}
+
+			return true;
+		}
+	};
+	IMPLEMENT_AI_INSTANT_TEST(FChildOf, "System.Mass.Relations.Hierarchy.ChildOf");
+
+	/** Testing individual policies being applied properly when relation object gets destroyed. */
+	struct FObjectDestroyedPolicy : FRelationTestBase
+	{
+		virtual bool InstantTest() override
+		{
+			const FMassArchetypeHandle OriginalArchetype = IntsArchetype;
+
+			const FMassEntityHandle SubjectEntity = EntityManager->CreateEntity(OriginalArchetype);
+			const FMassEntityHandle ObjectEntity = EntityManager->CreateEntity(OriginalArchetype);
+			RelationshipManagerPtr->CreateRelationInstance(RelationTypeHandle, SubjectEntity, ObjectEntity);
+
+			EntityManager->DestroyEntity(ObjectEntity);
+			EntityManager->FlushCommands();
+			AITEST_TRUE("The source entity is destroyed along with the relation object", EntityManager->IsEntityValid(SubjectEntity) == false);
+
+			return true;
+		}
+	};
+	IMPLEMENT_AI_INSTANT_TEST(FObjectDestroyedPolicy, "System.Mass.Relations.Policy.OnObjectDestroyed");
+
+	struct FSubjectDestroyedPolicy : FRelationTestBase
+	{
+		virtual bool InstantTest() override
+		{
+			const FMassArchetypeHandle OriginalArchetype = IntsArchetype;
+			{
+				FRelationTypeTraits LocalRelationTraits = CreateRelationTypeDescription("TestRelationA");
+				LocalRelationTraits.RoleTraits[static_cast<int32>(ERelationRole::Subject)].DestructionPolicy = ERemovalPolicy::Destroy;
+
+				const FTypeHandle LocalRelationTypeHandle = EntityManager->GetTypeManager().RegisterType(MoveTemp(LocalRelationTraits));
+
+				const FMassEntityHandle SubjectEntity = EntityManager->CreateEntity(OriginalArchetype);
+				const FMassEntityHandle ObjectEntity = EntityManager->CreateEntity(OriginalArchetype);
+				RelationshipManagerPtr->CreateRelationInstance(LocalRelationTypeHandle, SubjectEntity, ObjectEntity);
+
+				EntityManager->DestroyEntity(SubjectEntity);
+				EntityManager->FlushCommands();
+				AITEST_TRUE("The object entity is destroyed along with the relation source", EntityManager->IsEntityValid(ObjectEntity) == false);
+			}
+			{
+				FRelationTypeTraits LocalRelationTraits = CreateRelationTypeDescription("TestRelationB");
+				LocalRelationTraits.RoleTraits[static_cast<int32>(ERelationRole::Subject)].DestructionPolicy = ERemovalPolicy::CleanUp;
+				const FTypeHandle LocalRelationTypeHandle = EntityManager->GetTypeManager().RegisterType(MoveTemp(LocalRelationTraits));
+
+				const FMassEntityHandle SubjectEntity = EntityManager->CreateEntity(OriginalArchetype);
+				const FMassEntityHandle ObjectEntity = EntityManager->CreateEntity(OriginalArchetype);
+				RelationshipManagerPtr->CreateRelationInstance(LocalRelationTypeHandle, SubjectEntity, ObjectEntity);
+
+				EntityManager->DestroyEntity(SubjectEntity);
+				EntityManager->FlushCommands();
+				AITEST_TRUE("The object entity is still valid once the relation source gets destroyed", EntityManager->IsEntityValid(ObjectEntity));
+				// The relation has been cleaned up
+				AITEST_TRUE("The object entity is still valid once the the relation source gets destroyed", EntityManager->IsEntityValid(ObjectEntity));
+				TArray<FMassEntityHandle> Subjects = RelationshipManagerPtr->GetRelationSubjects(LocalRelationTypeHandle, ObjectEntity);
+				AITEST_EQUAL("Expected remaining sources count", Subjects.Num(), 0);
+			}
+			{
+				FRelationTypeTraits LocalRelationTraits = CreateRelationTypeDescription("TestRelationC");
+				LocalRelationTraits.RoleTraits[static_cast<int32>(ERelationRole::Subject)].DestructionPolicy = ERemovalPolicy::Splice;
+				// not that for this test to work the Object's Destruction policy can't be Destroy
+				// since then the LeafEntity (created below) will be destroyed as part of SubjectEntity's destruction
+				// before the new, patched relation gets created
+				LocalRelationTraits.RoleTraits[static_cast<int32>(ERelationRole::Object)].DestructionPolicy = ERemovalPolicy::Splice;
+				const FTypeHandle LocalRelationTypeHandle = EntityManager->GetTypeManager().RegisterType(MoveTemp(LocalRelationTraits));
+
+				const FMassEntityHandle SubjectEntity = EntityManager->CreateEntity(OriginalArchetype);
+				const FMassEntityHandle LeafEntity = EntityManager->CreateEntity(OriginalArchetype);
+				const FMassEntityHandle ObjectEntity = EntityManager->CreateEntity(OriginalArchetype);
+				RelationshipManagerPtr->CreateRelationInstance(LocalRelationTypeHandle, SubjectEntity, ObjectEntity);
+				RelationshipManagerPtr->CreateRelationInstance(LocalRelationTypeHandle, LeafEntity, SubjectEntity);
+
+				EntityManager->DestroyEntity(SubjectEntity);
+				EntityManager->FlushCommands();
+				AITEST_TRUE("The object entity is still valid once the relation source gets destroyed", EntityManager->IsEntityValid(ObjectEntity));
+				AITEST_TRUE("The leaf entity is still valid once the relation source gets destroyed", EntityManager->IsEntityValid(LeafEntity));
+				TArray<FMassEntityHandle> Subjects = RelationshipManagerPtr->GetRelationSubjects(LocalRelationTypeHandle, ObjectEntity);
+				AITEST_EQUAL("Expected remaining sources count", Subjects.Num(), 1);
+				AITEST_EQUAL("The leaf entity is the child now", Subjects[0], LeafEntity);
+			}
+
+			return true;
+		}
+	};
+	IMPLEMENT_AI_INSTANT_TEST(FSubjectDestroyedPolicy, "System.Mass.Relations.Policy.OnSubjectDestroyed");
+
+	struct FDestroyEverythingPolicy : FRelationTestBase
+	{
+		virtual bool InstantTest() override
+		{
+			const FMassArchetypeHandle OriginalArchetype = IntsArchetype;
+			{
+				const FMassEntityHandle SubjectEntity = EntityManager->CreateEntity(OriginalArchetype);
+				const FMassEntityHandle ObjectEntity = EntityManager->CreateEntity(OriginalArchetype);
+				RelationshipManagerPtr->CreateRelationInstance(RelationTypeHandle, SubjectEntity, ObjectEntity);
+
+				EntityManager->DestroyEntity(ObjectEntity);
+				EntityManager->FlushCommands();
+				AITEST_TRUE("The source entity is destroyed along with the relation object", EntityManager->IsEntityValid(SubjectEntity) == false);
+			}
+			{
+				const FMassEntityHandle SubjectEntity = EntityManager->CreateEntity(OriginalArchetype);
+				const FMassEntityHandle ObjectEntity = EntityManager->CreateEntity(OriginalArchetype);
+				RelationshipManagerPtr->CreateRelationInstance(RelationTypeHandle, SubjectEntity, ObjectEntity);
+
+				EntityManager->DestroyEntity(SubjectEntity);
+				EntityManager->FlushCommands();
+				AITEST_TRUE("The source entity is destroyed along with the relation object", EntityManager->IsEntityValid(SubjectEntity) == false);
+			}
+
+
+			return true;
+		}
+	};
+	IMPLEMENT_AI_INSTANT_TEST(FDestroyEverythingPolicy, "System.Mass.Relations.Policy.DestroyEverything");
+
+	struct FHierarchyTrivialCycle : FRelationTestBase
+	{
+		virtual bool InstantTest() override
+		{
+			const FMassArchetypeHandle OriginalArchetype = IntsArchetype;
+
+			const FMassEntityHandle Entity = EntityManager->CreateEntity(OriginalArchetype);
+			{
+				AITEST_SCOPED_CHECK("subject and object to be different", 1);
+				RelationshipManagerPtr->CreateRelationInstance(RelationTypeHandle, Entity, Entity);
+			}
+			const FMassArchetypeHandle FinalArchetype = EntityManager->GetArchetypeForEntity(Entity);
+			AITEST_TRUE("The entity has not changed archetypes", FinalArchetype == OriginalArchetype);
+
+			return true;
+		}
+	};
+	IMPLEMENT_AI_INSTANT_TEST(FHierarchyTrivialCycle, "System.Mass.Relations.Hierarchy.TrivialCycle");
+
+#if 0
+	// this test is disabled for now, no elaborate cycle detection in place yet
+	struct FHierarchyCycle : FRelationTestBase
+	{
+		virtual bool InstantTest() override
+		{
+			constexpr int32 NumEntities = 3;
+			const FMassArchetypeHandle OriginalArchetype = IntsArchetype;
+
+			TArray<FMassEntityHandle> Entities;
+			EntityManager->BatchCreateEntities(OriginalArchetype, NumEntities, Entities);
+
+			RelationshipManagerPtr->CreateRelationInstance(RelationTypeHandle, Entities[1], Entities[0]);
+			RelationshipManagerPtr->CreateRelationInstance(RelationTypeHandle, Entities[2], Entities[1]);
+			{
+				// this relation registration should fail as it would create a cycle
+				AITEST_SCOPED_CHECK("result in a cycle", 1);
+				RelationshipManagerPtr->CreateRelationInstance(RelationTypeHandle, Entities[0], Entities[2]);
+			}
+
+			return true;
+		}
+	};
+	IMPLEMENT_AI_INSTANT_TEST(FHierarchyCycle, "System.Mass.Relations.Hierarchy.Cycle");
+#endif
+
+	struct FDestructionPolicy : FRelationTestBase
+	{
+		virtual bool InstantTest() override
+		{
+			const FMassArchetypeHandle OriginalArchetype = IntsArchetype;
+
+			UScriptStruct* DestroyChildOnParentDestructionRelation = MakeRelationType("DestroyChildOnParentDestruction");
+			FRelationTypeTraits LocalRelationTraits = CreateRelationTypeDescription(DestroyChildOnParentDestructionRelation);
+			LocalRelationTraits.RoleTraits[static_cast<uint8>(ERelationRole::Object)].DestructionPolicy = ERemovalPolicy::Destroy;
+
+			FTypeHandle LocalRelationTypeHandle = EntityManager->GetTypeManager().RegisterType(MoveTemp(LocalRelationTraits));
+
+			TArray<FMassEntityHandle> Entities;
+			EntityManager->BatchCreateEntities(OriginalArchetype, 2, Entities);
+
+			RelationshipManagerPtr->CreateRelationInstance(LocalRelationTypeHandle, Entities[1], Entities[0]);
+			EntityManager->FlushCommands();
+
+			EntityManager->DestroyEntity(Entities[0]);
+
+			EntityManager->FlushCommands();
+			AITEST_TRUE("Destroying the parent destroys the child.", EntityManager->IsEntityActive(Entities[1]) == false);
+
+			return true;
+		}
+	};
+	IMPLEMENT_AI_INSTANT_TEST(FDestructionPolicy, "System.Mass.Relations.DestructionPolicy");
+
+	struct FExclusivity : FRelationTestBase
+	{
+		virtual bool InstantTest() override
+		{
+			constexpr int32 NumEntities = 3;
+			constexpr int32 OriginalParentIndex = 0;
+			constexpr int32 ChildIndex = 1;
+			constexpr int32 NewParentIndex = 2;
+			const FMassArchetypeHandle OriginalArchetype = IntsArchetype;
+
+			UScriptStruct* DestroyChildOnParentDestructionRelation = MakeRelationType("DestroyChildOnParentDestruction");
+			FRelationTypeTraits LocalRelationTraits = CreateRelationTypeDescription(DestroyChildOnParentDestructionRelation);
+			LocalRelationTraits.RoleTraits[static_cast<uint8>(ERelationRole::Subject)].DestructionPolicy = ERemovalPolicy::Destroy;
+
+			FTypeHandle LocalRelationTypeHandle = EntityManager->GetTypeManager().RegisterType(MoveTemp(LocalRelationTraits));
+
+			TArray<FMassEntityHandle> Entities;
+			EntityManager->BatchCreateEntities(OriginalArchetype, NumEntities, Entities);
+
+			RelationshipManagerPtr->CreateRelationInstance(LocalRelationTypeHandle, Entities[ChildIndex], Entities[OriginalParentIndex]);
+			// swapping parent
+			RelationshipManagerPtr->CreateRelationInstance(LocalRelationTypeHandle, Entities[ChildIndex], Entities[NewParentIndex]);
+
+			AITEST_TRUE("Original parent has no children", RelationshipManagerPtr->GetRelationSubjects(LocalRelationTypeHandle, Entities[OriginalParentIndex]).Num() == 0);
+			AITEST_TRUE("New parent has exactly one child", RelationshipManagerPtr->GetRelationSubjects(LocalRelationTypeHandle, Entities[NewParentIndex]).Num() == 1);
+			AITEST_TRUE("New parent's child matches expectations", RelationshipManagerPtr->GetRelationSubjects(LocalRelationTypeHandle, Entities[NewParentIndex])[0] == Entities[ChildIndex]);
+			AITEST_TRUE("The Child has parents", RelationshipManagerPtr->GetRelationObjects(LocalRelationTypeHandle, Entities[ChildIndex]).Num() > 0);
+			AITEST_TRUE("The Child knows about the new parent", RelationshipManagerPtr->GetRelationObjects(LocalRelationTypeHandle, Entities[ChildIndex])[0] == Entities[NewParentIndex]);
+
+			return true;
+		}
+	};
+	IMPLEMENT_AI_INSTANT_TEST(FExclusivity, "System.Mass.Relations.Exclusivity");
+
+	struct FInvalidRelations : FRelationTestBase
+	{
+		virtual bool InstantTest() override
+		{
+			constexpr int32 NumEntities = 3;
+			const FMassArchetypeHandle OriginalArchetype = IntsArchetype;
+
+			UScriptStruct* DestroyChildOnParentDestructionRelation = MakeRelationType("DestroyChildOnParentDestruction");
+			FRelationTypeTraits LocalRelationTraits = CreateRelationTypeDescription(DestroyChildOnParentDestructionRelation);
+			const FTypeHandle LocalRelationTypeHandle = EntityManager->GetTypeManager().RegisterType(MoveTemp(LocalRelationTraits));
+
+			TArray<FMassEntityHandle> Entities;
+			EntityManager->BatchCreateEntities(OriginalArchetype, NumEntities, Entities);
+
+			RelationshipManagerPtr->CreateRelationInstance(LocalRelationTypeHandle, Entities[0], Entities[1]);
+			EntityManager->FlushCommands();
+
+			// try creating relations with mismatching pairs
+			{
+				AITEST_SCOPED_CHECK("Relation Objects count need match Subjects", 1);
+				TArray<FMassEntityHandle> RelationInstances = RelationshipManagerPtr->CreateRelationInstances(LocalRelationTypeHandle
+					, TArrayView<FMassEntityHandle>(&Entities[0], 1), TArrayView<FMassEntityHandle>(&Entities[1], 2));
+				AITEST_EQUAL("A: Mismatching relation instances created count", RelationInstances.Num(), 0);
+			}
+			{
+				TArray<FMassEntityHandle> Subjects;
+				TArray<FMassEntityHandle> Objects;
+
+				AITEST_SCOPED_CHECK("Relation needs a valid Subject entity", 1);
+				Subjects.Add(FMassEntityHandle());
+				Objects.Add(Entities[0]);
+
+				AITEST_SCOPED_CHECK("Relation needs a valid Object entity", 1);
+				Subjects.Add(Entities[0]);
+				Objects.Add(FMassEntityHandle());
+				
+				AITEST_SCOPED_CHECK("Hierarchical relation requires the subject and object to be different", 1);
+				Subjects.Add(Entities[0]);
+				Objects.Add(Entities[0]);
+
+				// this one we expect to succeed:
+				Subjects.Add(Entities[1]);
+				Objects.Add(Entities[2]);
+
+				AITEST_SCOPED_CHECK("Relation between the two entities already exists", 1);
+				Subjects.Add(Entities[0]);
+				Objects.Add(Entities[1]);
+
+				TArray<FMassEntityHandle> RelationInstances = RelationshipManagerPtr->CreateRelationInstances(LocalRelationTypeHandle, Subjects, Objects);
+				AITEST_EQUAL("Expected number of valid relation instances created", RelationInstances.Num(), 1);
+
+				FMassRelationFragment* RelationFragment = EntityManager->GetFragmentDataPtr<FMassRelationFragment>(RelationInstances[0]);
+				AITEST_NOT_NULL(TEXT("The created relation entity has the relation fragment"), RelationFragment);
+				AITEST_TRUE("The created relation instance has the expected subject", RelationFragment->Subject == Entities[1]);
+				AITEST_TRUE("The created relation instance has the expected object", RelationFragment->Object == Entities[2]);
+			}
+
+			return true;
+		}
+	};
+	IMPLEMENT_AI_INSTANT_TEST(FInvalidRelations, "System.Mass.Relations.InvalidRelations");
+
+	struct FReRegisteringRelationType : FRelationTestBase
+	{
+		virtual bool InstantTest() override
+		{
+			{
+				FRelationTypeTraits LocalRelationTraits = CreateRelationTypeDescription("TestFragmentA");
+				EntityManager->GetTypeManager().RegisterType(MoveTemp(LocalRelationTraits));
+			}
+			{
+				AITEST_SCOPED_CHECK("Modifying relationship after registration done is not supported", 1);
+				FRelationTypeTraits LocalRelationTraits = CreateRelationTypeDescription("TestFragmentA");
+				// make it different
+				LocalRelationTraits.bHierarchical = !LocalRelationTraits.bHierarchical;
+				EntityManager->GetTypeManager().RegisterType(MoveTemp(LocalRelationTraits));
+			}
+			{
+				AITEST_SCOPED_CHECK("Modifying relationship after registration done is not supported", 1);
+				FRelationTypeTraits LocalRelationTraits = CreateRelationTypeDescription("TestFragmentB");
+				// register with existing name
+				EntityManager->GetTypeManager().RegisterType(MoveTemp(LocalRelationTraits));
+			}
+			{
+				AITEST_SCOPED_CHECK("Modifying relationship after registration done is not supported", 1);
+				FRelationTypeTraits LocalRelationTraits = CreateRelationTypeDescription("TestFragmentA");
+				// register existing traits - source and object elements are supposed to be unique
+				EntityManager->GetTypeManager().RegisterType(MoveTemp(LocalRelationTraits));
+			}
+
+			return true;
+		}
+	};
+	IMPLEMENT_AI_INSTANT_TEST(FReRegisteringRelationType, "System.Mass.Relations.Type.Registration");
+
+}// UE::Mass::Test::Relations
+
+UE_ENABLE_OPTIMIZATION_SHIP
+
+#undef LOCTEXT_NAMESPACE
diff --git a/Engine/Source/Developer/TextureFormat/Public/Interfaces/ITextureFormat.h b/Engine/Source/Developer/TextureFormat/Public/Interfaces/ITextureFormat.h
--- a/Engine/Source/Developer/TextureFormat/Public/Interfaces/ITextureFormat.h
+++ b/Engine/Source/Developer/TextureFormat/Public/Interfaces/ITextureFormat.h
@@ -424,7 +424,7 @@ struct FEncodedTextureDescription
 		MipTextureDescription.TopMipSizeX = TailFirstMipDims.X;
 		MipTextureDescription.TopMipSizeY = TailFirstMipDims.Y;
 		MipTextureDescription.TopMipVolumeSizeZ = TailFirstMipDims.Z;
-		MipTextureDescription.NumMips = NumMips - InRemoveCount;
+		MipTextureDescription.NumMips = NumMips - static_cast<uint8>(InRemoveCount);
 		if (InExtendedData && InExtendedData->NumMipsInTail && InRemoveCount >= NumMips - InExtendedData->NumMipsInTail)
 		{
 			// we must only ever get the first mip tail index!
diff --git a/Engine/Source/Runtime/MassEntity/Private/MassEntityBuilder.cpp b/Engine/Source/Runtime/MassEntity/Private/MassEntityBuilder.cpp
--- a/Engine/Source/Runtime/MassEntity/Private/MassEntityBuilder.cpp
+++ b/Engine/Source/Runtime/MassEntity/Private/MassEntityBuilder.cpp
@@ -280,23 +280,18 @@ FMassEntityHandle FEntityBuilder::Commit()
 	{
 		// we need to issue commands in this case
 		EntityManager->Defer().PushCommand<FMassDeferredCreateCommand>(
-					[ReservedEntityHandle = EntityHandle, SharedFragmentValues = CachedSharedFragmentValues, ArchetypeHandle = CachedArchetypeHandle, FragmentsCopy = Fragments](FMassEntityManager& Manager)
+					[ReservedEntityHandle = EntityHandle, SharedFragmentValues = CachedSharedFragmentValues
+					, ArchetypeHandle = CachedArchetypeHandle, FragmentsCopy = Fragments
+					, RelationsParams = RelationsParams](FMassEntityManager& Manager)
 					{
-						// creating creation context to block observers from triggering until the
-						// values are set with SetEntityFragmentValues. The lock gets auto-released at the end of the scope
-						// or persists, if anyone else has a shared ptr to it
-						TSharedRef<FMassEntityManager::FEntityCreationContext> CreationContext = Manager.GetOrMakeCreationContext();
-						Manager.BuildEntity(ReservedEntityHandle, ArchetypeHandle, SharedFragmentValues);
-						Manager.SetEntityFragmentValues(ReservedEntityHandle, FragmentsCopy);
+						FEntityBuilder::CreateEntityImpl(Manager, ReservedEntityHandle, ArchetypeHandle
+							, SharedFragmentValues, FragmentsCopy, RelationsParams);
 					});
 	}
 	else
 	{
-		// directly create the entity
-		TSharedRef<FMassEntityManager::FEntityCreationContext> CreationContext = EntityManager->GetOrMakeCreationContext();
-
-		EntityManager->BuildEntity(EntityHandle, CachedArchetypeHandle, CachedSharedFragmentValues);
-		EntityManager->SetEntityFragmentValues(EntityHandle, Fragments);
+		// directly create the entity, right now
+		CreateEntityImpl(*EntityManager, EntityHandle, CachedArchetypeHandle, CachedSharedFragmentValues, Fragments, RelationsParams);
 	}
 
 	State = EState::Committed;
@@ -304,6 +299,32 @@ FMassEntityHandle FEntityBuilder::Commit()
 	return EntityHandle;
 }
 
+void FEntityBuilder::CreateEntityImpl(FMassEntityManager& EntityManager, FMassEntityHandle EntityHandle, const FMassArchetypeHandle& ArchetypeHandle
+	, const FMassArchetypeSharedFragmentValues& SharedFragmentValues, const TConstArrayView<FInstancedStruct> Fragments, TConstArrayView<FPendingRelationParams> RelationsParams)
+{
+	// creating creation context to block observers from triggering until the
+	// values are set with SetEntityFragmentValues. The lock gets auto-released at the end of the scope
+	// or persists, if anyone else has a shared ptr to it
+	TSharedRef<FMassEntityManager::FEntityCreationContext> CreationContext = EntityManager.GetOrMakeCreationContext();
+
+	EntityManager.BuildEntity(EntityHandle, ArchetypeHandle, SharedFragmentValues);
+	EntityManager.SetEntityFragmentValues(EntityHandle, Fragments);
+	for (const FPendingRelationParams& Relation : RelationsParams)
+	{
+		switch (Relation.OtherEntityRole)
+		{
+		case ERelationRole::Subject:
+			EntityManager.GetRelationManager().CreateRelationInstance(Relation.RelationTypeHandle, Relation.OtherEntity, EntityHandle);
+			break;
+		case ERelationRole::Object:
+			EntityManager.GetRelationManager().CreateRelationInstance(Relation.RelationTypeHandle, EntityHandle, Relation.OtherEntity);
+			break;
+		default:
+			ensureMsgf(false, TEXT("Unhandled value %s"), *LexToString(Relation.OtherEntityRole));
+		}
+	}
+}
+
 FMassEntityHandle FEntityBuilder::CommitAndReprepare()
 {
 	FMassEntityHandle CreatedEntity = Commit();
@@ -337,6 +358,7 @@ void FEntityBuilder::Reset(const bool bReleaseEntityHandleIfReserved)
 		Fragments.Reset();
 		SharedFragments.Reset();
 		ConstSharedFragments.Reset();
+		RelationsParams.Reset();
 	}
 }
 
@@ -575,4 +597,21 @@ FEntityBuilder& FEntityBuilder::Add(TNotNull<const UScriptStruct*> ElementType)
 	return *this;
 }
 
+FEntityBuilder& FEntityBuilder::AddRelation(const FTypeHandle RelationTypeHandle, const FMassEntityHandle OtherEntity, const ERelationRole InputEntityRole)
+{
+	RelationsParams.Add({RelationTypeHandle, OtherEntity, InputEntityRole});
+	return *this;
+}
+
+void FEntityBuilder::ForEachRelation(const TFunctionRef<bool(FPendingRelationParams&)>& Operator)
+{
+	for (int32 ElementIndex = RelationsParams.Num() - 1; ElementIndex >= 0; --ElementIndex)
+	{
+		if (Operator(RelationsParams[ElementIndex]) == false)
+		{
+			RelationsParams.RemoveAt(ElementIndex, EAllowShrinking::No);
+		}
+	}
+}
+
 } // namespace UE::Mass
\ No newline at end of file
diff --git a/Engine/Source/Runtime/MassEntity/Private/MassEntityManager.cpp b/Engine/Source/Runtime/MassEntity/Private/MassEntityManager.cpp
--- a/Engine/Source/Runtime/MassEntity/Private/MassEntityManager.cpp
+++ b/Engine/Source/Runtime/MassEntity/Private/MassEntityManager.cpp
@@ -62,6 +62,7 @@ namespace UE::Mass::Private
 FMassEntityManager::FMassEntityManager(UObject* InOwner)
 	: ObserverManager(*this)
 	, TypeManager(new UE::Mass::FTypeManager(*this))
+	, RelationManager(*this)
 	, Owner(InOwner)
 {
 #if WITH_MASSENTITY_DEBUG
@@ -246,6 +247,17 @@ void FMassEntityManager::PostInitialize()
 	ensureMsgf(InitializationState == EInitializationState::Initialized,
 		TEXT("This needs to be done after all the subsystems have been initialized since some processors might want to access"
 			" them during processors' initialization"));
+
+	TypeManager->RegisterBuiltInTypes();
+	// now hook-in all relation observers
+	// note that we're doing it only after RegisterBuiltInTypes is done, as opposed to doing it on
+	// every type as it gets added, because there are ways to override the traits of built-in types.
+	// Once RegisterBuiltInTypes is done the traits are set in stone, and we can handle the types.
+	for (UE::Mass::FTypeManager::FTypeInfoConstIterator It = TypeManager->MakeIterator(); It; ++It)
+	{
+		OnNewTypeRegistered(It->Key);
+	}
+
 	ObserverManager.Initialize();
 }
 
@@ -2069,6 +2081,32 @@ UE::Mass::FEntityBuilder FMassEntityManager::MakeEntityBuilder()
 	return UE::Mass::FEntityBuilder(AsShared());
 }
 
+void FMassEntityManager::OnNewTypeRegistered(UE::Mass::FTypeHandle RegisteredTypeHandle)
+{
+	const UE::Mass::FTypeInfo* NewTypeInfo = TypeManager->GetTypeInfo(RegisteredTypeHandle);
+	if (testableEnsureMsgf(NewTypeInfo, TEXT("%hs input handle doesn't represent a registered type"), __FUNCTION__))
+	{
+		if (const UE::Mass::FRelationTypeTraits* RelationTypeTraits = NewTypeInfo->GetAsRelationTraits())
+		{
+			OnRelationTypeRegistered(RegisteredTypeHandle, *RelationTypeTraits);
+		}
+	}
+}
+
+void FMassEntityManager::OnRelationTypeRegistered(UE::Mass::FTypeHandle RegisteredTypeHandle, const UE::Mass::FRelationTypeTraits& RelationTypeTraits)
+{
+	if (RelationTypeTraits.RegisterObserversDelegate.IsSet() == false
+		|| RelationTypeTraits.RegisterObserversDelegate(*this) == true)
+	{
+		RelationManager.OnRelationTypeRegistered(RegisteredTypeHandle, RelationTypeTraits);
+	}
+}
+
+bool FMassEntityManager::BatchCreateRelations(const UE::Mass::FTypeHandle RelationTypeHandle, TArrayView<FMassEntityHandle> Subjects, TArrayView<FMassEntityHandle> Objects)
+{
+	return RelationManager.CreateRelationInstances(RelationTypeHandle, Subjects, Objects).Num() > 0;
+}
+
 //-----------------------------------------------------------------------------
 // DEBUG stuff
 //-----------------------------------------------------------------------------
diff --git a/Engine/Source/Runtime/MassEntity/Private/MassEntityModule.cpp b/Engine/Source/Runtime/MassEntity/Private/MassEntityModule.cpp
--- a/Engine/Source/Runtime/MassEntity/Private/MassEntityModule.cpp
+++ b/Engine/Source/Runtime/MassEntity/Private/MassEntityModule.cpp
@@ -5,6 +5,7 @@
 #include "MassDebugger.h"
 #include "MassProcessingTypes.h"
 #include "Modules/ModuleManager.h"
+#include "Relations/MassChildOf.h"
 
 #define LOCTEXT_NAMESPACE "Mass"
 
@@ -28,6 +29,8 @@ void FMassEntityModule::StartupModule()
 #else
 	UE_LOG(LogMass, Log, TEXT("MassEntity running in game thread."));
 #endif // MASS_DO_PARALLEL
+
+	UE::Mass::Relations::RegisterChildOfRelation();
 }
 
 void FMassEntityModule::ShutdownModule()
diff --git a/Engine/Source/Runtime/MassEntity/Private/MassEntityRelations.cpp b/Engine/Source/Runtime/MassEntity/Private/MassEntityRelations.cpp
new file mode 100644
--- /dev/null
+++ b/Engine/Source/Runtime/MassEntity/Private/MassEntityRelations.cpp
@@ -0,0 +1,58 @@
+// Copyright Epic Games, Inc. All Rights Reserved.
+
+#include "MassEntityRelations.h"
+#include "MassEntityManager.h"
+#include "MassRelationObservers.h"
+
+namespace UE::Mass
+{
+	//-----------------------------------------------------------------------------
+	// FRelationTypeTraits
+	//-----------------------------------------------------------------------------
+	FRelationTypeTraits::FRelationTypeTraits(const TNotNull<const UScriptStruct*> InRelationTagType)
+		: RelationTagType(InRelationTagType)
+		, RelationFragmentType(FMassRelationFragment::StaticStruct())
+	{
+		// at the moment UMassRelationEntityCreation doesn't do anything, so we don't plug it in
+		// RelationEntityCreationObserverClass = UMassRelationEntityCreation::StaticClass();
+		RelationEntityDestructionObserverClass = UMassRelationEntityDestruction::StaticClass();
+
+		SubjectEntityDestructionObserverClass
+			= ObjectEntityDestructionObserverClass = UMassRelationRoleDestruction::StaticClass();
+	}
+
+	FRelationTypeTraits::FRelationTypeTraits(const FRelationTypeTraits& Other, const TNotNull<const UScriptStruct*> NewRelationTagType)
+		: FRelationTypeTraits(Other)
+	{
+		RelationTagType = NewRelationTagType;
+	}
+
+	void FRelationTypeTraits::SetDebugInFix(FString&& InFix)
+	{
+#if WITH_MASSENTITY_DEBUG
+		DebugInFix = MoveTemp(InFix);
+#endif // WITH_MASSENTITY_DEBUG
+	}
+
+#if WITH_MASSENTITY_DEBUG
+	FString FRelationTypeTraits::DebugDescribeRelation(FMassEntityHandle A, FMassEntityHandle B) const
+	{
+		return FString::Printf(TEXT("[%s] %s [%s]"), *A.DebugGetDescription(), *DebugInFix, *B.DebugGetDescription());
+	}
+#endif // WITH_MASSENTITY_DEBUG
+}
+
+//-----------------------------------------------------------------------------
+// FMassRelationRoleInstanceHandle
+//-----------------------------------------------------------------------------
+FMassEntityHandle FMassRelationRoleInstanceHandle::GetRoleEntityHandle(const FMassEntityManager& EntityManager) const
+{
+	const int32 EntityIndex = GetRoleEntityIndex();
+	return EntityManager.CreateEntityIndexHandle(EntityIndex);
+}
+
+FMassEntityHandle FMassRelationRoleInstanceHandle::GetRelationEntityHandle(const FMassEntityManager& EntityManager) const
+{
+	const int32 EntityIndex = GetRelationEntityIndex();
+	return EntityManager.CreateEntityIndexHandle(EntityIndex);
+}
diff --git a/Engine/Source/Runtime/MassEntity/Private/MassRelationManager.cpp b/Engine/Source/Runtime/MassEntity/Private/MassRelationManager.cpp
new file mode 100644
--- /dev/null
+++ b/Engine/Source/Runtime/MassEntity/Private/MassRelationManager.cpp
@@ -0,0 +1,562 @@
+// Copyright Epic Games, Inc. All Rights Reserved.
+
+#include "MassRelationManager.h"
+#include "MassEntityBuilder.h"
+#include "MassEntityTypes.h"
+#include "MassEntityManager.h"
+#include "MassRelationObservers.h"
+#include "MassExecutionContext.h"
+#include "Interfaces/ITextureFormat.h"
+#include "MassCommands.h"
+#include "MassArchetypeGroupCommands.h"
+#include "Algo/RemoveIf.h"
+
+DECLARE_CYCLE_STAT(TEXT("Mass Relation IsSubject"), STAT_Mass_IsSubject, STATGROUP_Mass);
+DECLARE_CYCLE_STAT(TEXT("Mass Relation IsSubject Recursive"), STAT_Mass_IsSubjectRecursive, STATGROUP_Mass);
+
+namespace UE::Mass
+{
+	namespace Relations::Private
+	{
+		inline void HandleRole(const FMassEntityManager& EntityManager, FRelationData& RelationData, const FMassEntityHandle& OperatorEntity, UE::Mass::ERelationRole Role
+			, const FMassEntityHandle& RoleEntity, const FMassEntityHandle& RelationEntity
+			, TArray<FMassEntityHandle>& InOutRelationEntitiesToDestroy)
+		{
+			const int32 RoleIndex = static_cast<int32>(Role);
+			if (RelationData.Traits.RoleTraits[RoleIndex].RequiresExternalMapping == EExternalMappingRequired::Yes)
+			{
+				TArray<FMassRelationRoleInstanceHandle>& RoleData = RelationData.RoleMap.FindOrAdd(OperatorEntity)[RoleIndex];
+
+				const FMassRelationRoleInstanceHandle NewInstanceHandle = FMassRelationRoleInstanceHandle::Create(Role, RoleEntity, RelationEntity);
+
+				// if objects are exclusive we need to destroy the previous instance
+				if (RelationData.Traits.RoleTraits[RoleIndex].bExclusive
+					&& RoleData.Num()
+					&& RoleData[0] != NewInstanceHandle)
+				{
+					FMassRelationRoleInstanceHandle PreviousInstance = RoleData.Pop(EAllowShrinking::No);
+					const FMassEntityHandle PreviousRelationEntity = PreviousInstance.GetRelationEntityHandle(EntityManager);
+					InOutRelationEntitiesToDestroy.Add(PreviousRelationEntity);
+				}
+
+#if WITH_MASSENTITY_DEBUG
+				const int32 ExistingIndex = RoleData.Find(NewInstanceHandle);
+				ensureMsgf(ExistingIndex == INDEX_NONE, TEXT("Given relation instance handle is already present in role mapping"));
+#endif // WITH_MASSENTITY_DEBUG
+
+				RoleData.Add(NewInstanceHandle);
+			}
+		}
+
+		void AddRoleInstance(const FMassEntityManager& EntityManager, FRelationData& RelationData, const FMassEntityHandle Subject, const FMassEntityHandle Object, const FMassEntityHandle RelationEntityHandle
+			, TArray<FMassEntityHandle>& InOutRelationEntitiesToDestroy)
+		{
+			HandleRole(EntityManager, RelationData, Subject, ERelationRole::Object, Object, RelationEntityHandle, InOutRelationEntitiesToDestroy);
+			HandleRole(EntityManager, RelationData, Object, ERelationRole::Subject, Subject, RelationEntityHandle, InOutRelationEntitiesToDestroy);
+		}
+	}
+
+	struct FScopedRecursiveLimit
+	{
+		FScopedRecursiveLimit(int32& Limit)
+			: LimitRef(Limit)
+		{
+			++LimitRef;
+			check(LimitRef < 10);
+		}
+		~FScopedRecursiveLimit()
+		{
+			--LimitRef;
+		}
+		int32& LimitRef;
+	};
+
+	//-----------------------------------------------------------------------------
+	// FRelationData
+	//-----------------------------------------------------------------------------
+	FRelationData::FRelationData(const FRelationTypeTraits& InTraits)
+		: Traits(InTraits)
+	{
+	}
+
+	TArray<FMassEntityHandle> FRelationData::GetParticipants(const FMassEntityManager& EntityManager, const FMassEntityHandle RoleEntity, const ERelationRole QueriedRole) const
+	{
+		const int32 QueriedRoleIndex = static_cast<int32>(QueriedRole);
+
+		ensureMsgf(Traits.RoleTraits[QueriedRoleIndex].RequiresExternalMapping == EExternalMappingRequired::Yes
+			, TEXT("Fetching relation participants or role %s, while role traits explicitly prevent generic mapping")
+			, *LexToString(QueriedRole));
+
+		TArray<FMassEntityHandle> ReturnSubjects;
+		if (const FRelationInstanceData* InstanceData = RoleMap.Find(RoleEntity))
+		{
+			ReturnSubjects.Reserve((*InstanceData)[QueriedRoleIndex].Num());
+			for (const FMassRelationRoleInstanceHandle RelationInstanceHandle : (*InstanceData)[QueriedRoleIndex])
+			{
+				const FMassEntityHandle RoleEntityHandle = RelationInstanceHandle.GetRoleEntityHandle(EntityManager);
+				if (RoleEntityHandle.IsValid())
+				{
+					ReturnSubjects.Add(RoleEntityHandle);
+				}
+			}
+		}
+		return ReturnSubjects;
+	}
+
+	//-----------------------------------------------------------------------------
+	// FRelationManager::FHierarchyEntitiesContainer
+	//-----------------------------------------------------------------------------
+	void FRelationManager::FHierarchyEntitiesContainer::StoreUnique(const uint32 Depth, TArray<FMassEntityHandle>& InOutArray)
+	{
+		for (int32 HandleIndex = InOutArray.Num() - 1; HandleIndex >= 0; --HandleIndex)
+		{
+			bool bIsAlreadyInSet = false;
+			ExistingElements.FindOrAdd(InOutArray[HandleIndex], &bIsAlreadyInSet);
+			if (bIsAlreadyInSet)
+			{
+				InOutArray.RemoveAtSwap(HandleIndex);
+			}
+		}
+
+		if (InOutArray.Num())
+		{
+			if (static_cast<uint32>(ContainerPerLevel.Num()) <= Depth)
+			{
+				ContainerPerLevel.AddDefaulted(Depth - ContainerPerLevel.Num() + 1);
+			}
+			ContainerPerLevel[Depth].Append(InOutArray);
+		}
+	}
+
+	void FRelationManager::FHierarchyEntitiesContainer::StoreUnique(const uint32 Depth, const FMassEntityHandle Handle)
+	{
+		bool bIsAlreadyInSet = false;
+		ExistingElements.FindOrAdd(Handle, &bIsAlreadyInSet);
+		if (bIsAlreadyInSet == false)
+		{
+			if (static_cast<uint32>(ContainerPerLevel.Num()) <= Depth)
+			{
+				ContainerPerLevel.AddDefaulted(Depth - ContainerPerLevel.Num() + 1);
+			}
+			ContainerPerLevel[Depth].Add(Handle);
+		}
+	}
+
+	//-----------------------------------------------------------------------------
+	// FRelationManager
+	//-----------------------------------------------------------------------------
+	FRelationManager::FRelationManager(FMassEntityManager& InEntityManager)
+		: EntityManager(InEntityManager)
+		, TypeManager(InEntityManager.GetTypeManager())
+	{
+	}
+
+	TArray<FMassEntityHandle> FRelationManager::CreateRelationInstances(const FTypeHandle RelationTypeHandle, TArrayView<FMassEntityHandle> Subjects, TArrayView<FMassEntityHandle> Objects)
+	{
+		if (!ensureMsgf(RelationTypeHandle.IsValid(), TEXT("Invalid RelationTypeHandle passed to %hs"), __FUNCTION__))
+		{
+			return {};
+		}
+		if (!ensureMsgf(Subjects.Num(), TEXT("Relation needs a valid Subjects entity")) 
+			|| !ensureMsgf(Objects.Num(), TEXT("Relation needs a valid Objects entity"))
+			// we allow number mismatch if there's a single object or subject, then we create 1-to-many relation
+			|| !testableEnsureMsgf(Objects.Num() == Subjects.Num(), TEXT("Relation Objects count need match Subjects")))
+		{
+			return {};
+		}
+
+		FRelationData& RelationData = GetRelationDataChecked(RelationTypeHandle);
+
+		int32 NumOfRelations = Subjects.Num();
+		for (int32 PairIndex = 0; PairIndex < NumOfRelations; )
+		{
+			const FMassEntityHandle Object = Objects[PairIndex];
+			const FMassEntityHandle Subject = Subjects[PairIndex];
+			bool bPairValid = true;
+
+			if (!testableEnsureMsgf(Subject.IsValid(), TEXT("Relation needs a valid Subject entity")) 
+				|| !testableEnsureMsgf(Object.IsValid(), TEXT("Relation needs a valid Object entity")))
+			{
+				bPairValid = false;
+			}
+			if (!testableEnsureMsgf(RelationData.Traits.bHierarchical == false || Subject != Object
+				, TEXT("Hierarchical relation requires the subject and object to be different")))
+			{
+				bPairValid = false;
+			}
+			if (!testableEnsureMsgf(IsSubjectOfRelation(RelationData, Subject, Object) == false
+				, TEXT("Relation between the two entities already exists")))
+			{
+				bPairValid = false;
+			}
+
+			if (bPairValid)
+			{
+				++PairIndex;
+			}
+			else
+			{
+				if (PairIndex + 1 < NumOfRelations)
+				{
+					Swap(Objects[PairIndex], Objects[NumOfRelations - 1]);
+					Swap(Subjects[PairIndex], Subjects[NumOfRelations - 1]);
+				}
+				--NumOfRelations;
+			}
+		}
+
+		if (NumOfRelations == 0)
+		{
+			return {};
+		}
+
+		// some pairs have been removed, we need to update the views and pretend the invalid pairs were never there
+		if (NumOfRelations != Subjects.Num())
+		{
+			Objects = MakeArrayView(&Objects[0], NumOfRelations);
+			Subjects = MakeArrayView(&Subjects[0], NumOfRelations);
+		}
+
+		FEntityBuilder EntityBuilder(EntityManager);
+		FMassRelationFragment& RelationFragment = EntityBuilder.Add_GetRef<FMassRelationFragment>();
+		TNotNull<const UScriptStruct*> RelationTypeTag = RelationTypeHandle.GetScriptStruct();
+		EntityBuilder.Add(RelationTypeTag);
+
+		TSharedRef<UE::Mass::ObserverManager::FObserverLock> ObserversLock = EntityManager.GetOrMakeObserversLock();
+
+		TArray<FMassEntityHandle> CreatedRelationEntities;
+		TArray<FMassEntityHandle> RelationEntitiesToDestroy;
+
+		for (int32 PairIndex = 0; PairIndex < NumOfRelations; ++PairIndex)
+		{
+			const FMassEntityHandle Object = Objects[FMath::Min(Objects.Num() - 1, PairIndex)];
+			const FMassEntityHandle Subject = Subjects[FMath::Min(Subjects.Num() - 1, PairIndex)];
+
+			RelationFragment.Subject = Subject;
+			RelationFragment.Object = Object;
+
+			CreatedRelationEntities.Add(EntityBuilder.CommitAndReprepare());
+			const FMassEntityHandle RelationEntityHandle = CreatedRelationEntities.Last();
+
+			Relations::Private::AddRoleInstance(EntityManager, RelationData, Subject, Object, RelationEntityHandle, RelationEntitiesToDestroy);
+		}
+
+		if (RelationEntitiesToDestroy.Num())
+		{
+			EntityManager.Defer().DestroyEntities(MoveTemp(RelationEntitiesToDestroy));
+		}
+
+		TNotNull<const UScriptStruct*> SubjectElement = RelationData.Traits.RoleTraits[static_cast<int32>(ERelationRole::Subject)].Element
+			? TNotNull<const UScriptStruct*>(RelationData.Traits.RoleTraits[static_cast<int32>(ERelationRole::Subject)].Element)
+			: RelationTypeTag;
+		TNotNull<const UScriptStruct*> ObjectElement = RelationData.Traits.RoleTraits[static_cast<int32>(ERelationRole::Object)].Element
+			? TNotNull<const UScriptStruct*>(RelationData.Traits.RoleTraits[static_cast<int32>(ERelationRole::Object)].Element)
+			: RelationTypeTag;
+
+		if (EntityManager.IsProcessing())
+		{
+			TUniquePtr<FMassCommandAddElement> ElementAddingCommand = MakeUnique<FMassCommandAddElement>(SubjectElement);
+
+			ElementAddingCommand->Add(Subjects);
+
+			if (ObjectElement != SubjectElement)
+			{
+				EntityManager.Defer().PushUniqueCommand(MoveTemp(ElementAddingCommand));
+				ElementAddingCommand = MakeUnique<FMassCommandAddElement>(ObjectElement);
+			}
+			
+			ElementAddingCommand->Add(Objects);
+
+			EntityManager.Defer().PushUniqueCommand(MoveTemp(ElementAddingCommand));
+		}
+		else
+		{
+			EntityManager.AddElementToEntities(Objects, ObjectElement);
+			EntityManager.AddElementToEntities(Subjects, SubjectElement);
+		}
+
+		if (RelationData.Traits.bHierarchical)
+		{
+			FHierarchyEntitiesContainer SubSubjects;
+			uint32 MinDepth = static_cast<uint32>(-1);
+
+			for (int32 PairIndex = 0; PairIndex < Subjects.Num(); ++PairIndex)
+			{
+				const FMassEntityHandle Object = Objects[FMath::Min(Objects.Num() - 1, PairIndex)];
+				const FMassEntityHandle Subject = Subjects[FMath::Min(Subjects.Num() - 1, PairIndex)];
+
+				FArchetypeGroupHandle ObjectsGroup = EntityManager.GetGroupForEntity(Object, RelationData.Traits.RegisteredGroupType);
+				if (ObjectsGroup.IsValid() == false)
+				{
+					ObjectsGroup = FArchetypeGroupHandle(RelationData.Traits.RegisteredGroupType, 0);
+				}
+
+				const uint32 StartingDepth = ObjectsGroup.GetGroupID();
+				MinDepth = FMath::Min(StartingDepth, MinDepth);
+
+				SubSubjects.StoreUnique(StartingDepth, Object);
+				SubSubjects.StoreUnique(StartingDepth + 1, Subject);
+
+				GatherHierarchy(RelationData, Subject, SubSubjects, StartingDepth + 2);
+			}
+
+			if (EntityManager.IsProcessing())
+			{
+				for (uint32 Depth = MinDepth; Depth < SubSubjects.Num(); ++Depth)
+				{
+					FArchetypeGroupHandle GroupHandle(RelationData.Traits.RegisteredGroupType, Depth);
+
+					EntityManager.Defer().PushCommand<UE::Mass::Command::FGroupEntities>(GroupHandle, MoveTemp(SubSubjects[Depth]));
+				}
+			}
+			else
+			{
+				for (uint32 Depth = MinDepth; Depth < SubSubjects.Num(); ++Depth)
+				{
+					FArchetypeGroupHandle GroupHandle(RelationData.Traits.RegisteredGroupType, Depth);
+
+					TArray<FMassArchetypeEntityCollection> CollectionsAtLevel;
+					UE::Mass::Utils::CreateEntityCollections(EntityManager, SubSubjects[Depth]
+						, FMassArchetypeEntityCollection::EDuplicatesHandling::FoldDuplicates
+						, CollectionsAtLevel);
+
+					EntityManager.BatchGroupEntities(GroupHandle, CollectionsAtLevel);
+				}
+			}
+		}
+
+		return CreatedRelationEntities;
+	}
+
+	bool FRelationManager::DestroyRelationInstance(const FTypeHandle RelationTypeHandle, const FMassEntityHandle Subject, const FMassEntityHandle Object)
+	{
+		const FRelationData& RelationData = GetRelationDataChecked(RelationTypeHandle);
+		ensure(RelationData.Traits.RoleTraits[0].RequiresExternalMapping != EExternalMappingRequired::No
+			|| RelationData.Traits.RoleTraits[1].RequiresExternalMapping != EExternalMappingRequired::No);
+
+		if (const FRelationData::FRelationInstanceData* InstanceData = RelationData.RoleMap.Find(Subject))
+		{
+			for (const FMassRelationRoleInstanceHandle RelationInstanceHandle : (*InstanceData)[static_cast<uint32>(ERelationRole::Object)])
+			{
+				if (RelationInstanceHandle.GetRoleEntityIndex() == Object.Index)
+				{
+					const FMassEntityHandle RelationEntityHandle = RelationInstanceHandle.GetRelationEntityHandle(EntityManager);
+					if (RelationEntityHandle.IsValid())
+					{
+						EntityManager.DestroyEntity(RelationEntityHandle);
+						return true;
+					}
+				}
+			}
+		}
+		return false;
+	}
+
+	bool FRelationManager::DestroyRelationInstance(const FMassRelationRoleInstanceHandle RelationHandle) const
+	{
+		const FMassEntityHandle RelationEntityHandle = RelationHandle.GetRelationEntityHandle(EntityManager);
+		if (RelationEntityHandle.IsValid())
+		{
+			EntityManager.DestroyEntity(RelationEntityHandle);
+			return true;
+		}
+		return false;
+	}
+
+	TArray<FMassEntityHandle> FRelationManager::GetRelationSubjects(const FRelationData& RelationData, const FMassEntityHandle ObjectEntity) const
+	{
+		return RelationData.GetParticipants(EntityManager, ObjectEntity, ERelationRole::Subject);
+	}
+
+	TArray<FMassEntityHandle> FRelationManager::GetRelationObjects(const FRelationData& RelationData, const FMassEntityHandle SubjectEntity) const
+	{
+		return RelationData.GetParticipants(EntityManager, SubjectEntity, ERelationRole::Object);
+	}
+
+	void FRelationManager::GatherHierarchy(const FRelationData& RelationData, const FMassEntityHandle SubjectHandle, FHierarchyEntitiesContainer& SubSubjects, const uint32 Depth) const
+	{
+		// getting subject's subjects
+		TArray<FMassEntityHandle> LocalSubjects = GetRelationSubjects(RelationData, SubjectHandle);
+
+		if (LocalSubjects.Num() > 0)
+		{
+			// this call will store the handles in LocalSubjects, but only the ones that have not been
+			// stored earlier. This prevents the function getting stuck due to cycles.
+			SubSubjects.StoreUnique(Depth, LocalSubjects);
+			for (const FMassEntityHandle& SubSubjectHandle : LocalSubjects)
+			{
+				GatherHierarchy(RelationData, SubSubjectHandle, SubSubjects, Depth + 1);
+			}
+		}
+	}
+
+	void FRelationManager::GetRelationEntities(TConstArrayView<FMassRelationRoleInstanceHandle> RelationEntitiesContainer, TArray<FMassEntityHandle>& InOutEntityHandles) const
+	{
+		InOutEntityHandles.Reserve(RelationEntitiesContainer.Num() + InOutEntityHandles.Num());
+		for (const FMassRelationRoleInstanceHandle& Handle : RelationEntitiesContainer)
+		{
+			InOutEntityHandles.Add(Handle.GetRelationEntityHandle(EntityManager));
+		}
+	}
+
+	TArray<FMassEntityHandle> FRelationManager::GetRelationEntities(TConstArrayView<FMassRelationRoleInstanceHandle> RelationEntitiesContainer) const
+	{
+		TArray<FMassEntityHandle> EntityHandles;
+		GetRelationEntities(RelationEntitiesContainer, EntityHandles);
+		return EntityHandles;
+	}
+
+	void FRelationManager::GetRoleEntities(TConstArrayView<FMassRelationRoleInstanceHandle> RelationEntitiesContainer, TArray<FMassEntityHandle>& InOutEntityHandles) const
+	{
+		InOutEntityHandles.Reserve(RelationEntitiesContainer.Num() + InOutEntityHandles.Num());
+		for (const FMassRelationRoleInstanceHandle& Handle : RelationEntitiesContainer)
+		{
+			InOutEntityHandles.Add(Handle.GetRoleEntityHandle(EntityManager));
+		}
+	}
+
+	TArray<FMassEntityHandle> FRelationManager::GetRoleEntities(TConstArrayView<FMassRelationRoleInstanceHandle> RelationEntitiesContainer) const
+	{
+		TArray<FMassEntityHandle> EntityHandles;
+		GetRoleEntities(RelationEntitiesContainer, EntityHandles);
+		return EntityHandles;
+	}
+
+	FRelationData& FRelationManager::CreateRelationData(const FTypeHandle RelationTypeHandle)
+	{
+		FRelationData* RelationData = RelationsDataMap.Find(RelationTypeHandle);
+		if (ensureMsgf(RelationData == nullptr
+			, TEXT("%hs: relation of type %s already registered"), __FUNCTION__, *RelationTypeHandle.GetFName().ToString()))
+		{
+			FRelationTypeTraits Traits = TypeManager.GetRelationTypeChecked(RelationTypeHandle);
+
+			// @todo we should check every role's traits and see if RequiresExternalMapping == EExternalMappingRequired::No,
+			// and if so, make sure all the other required bits are provided (like "other role getter").
+
+			RelationData = &RelationsDataMap.Add(RelationTypeHandle, Traits);
+		}
+
+		return *RelationData;
+	}
+
+	bool FRelationManager::IsSubjectOfRelation(const FTypeHandle RelationTypeHandle, const FMassEntityHandle Subject, const FMassEntityHandle Object) const
+	{
+		check(RelationTypeHandle.IsValid());
+		if (const FRelationData* RelationData = GetRelationData(RelationTypeHandle))
+		{
+			return IsSubjectOfRelation(*RelationData, Subject, Object);
+		}
+		return false;
+	}
+
+	bool FRelationManager::IsSubjectOfRelation(const FRelationData& RelationData, const FMassEntityHandle Subject, const FMassEntityHandle Object) const
+	{
+		SCOPE_CYCLE_COUNTER(STAT_Mass_IsSubject);
+
+		if (Subject != Object
+			&& Subject.IsValid()
+			&& Object.IsValid())
+		{
+			if (const FRelationData::FRelationInstanceData* SubjectData = RelationData.RoleMap.Find(Subject))
+			{
+				for (const FMassRelationRoleInstanceHandle& RelationInstanceHandle : (*SubjectData)[static_cast<int32>(ERelationRole::Object)])
+				{
+					if (RelationInstanceHandle.GetRoleEntityHandle(EntityManager) == Object)
+					{
+						return true;
+					}
+				}
+			}
+		}
+
+		return false;
+	}
+
+	bool FRelationManager::IsSubjectOfRelationRecursive(const FTypeHandle RelationTypeHandle, const FMassEntityHandle Subject, const FMassEntityHandle Object) const
+	{
+		check(RelationTypeHandle.IsValid());
+		if (const FRelationData* RelationData = GetRelationData(RelationTypeHandle))
+		{
+			return IsSubjectOfRelationRecursive(*RelationData, Subject, Object);
+		}
+		return false;
+	}
+
+	bool FRelationManager::IsSubjectOfRelationRecursive(const FRelationData& RelationData, const FMassEntityHandle Subject, const FMassEntityHandle Object) const
+	{
+		// @todo the function can be optimized based on information whether Subject or Object is exclusive. If both are, it can be optimized even further.
+		// for now naive implementation
+		SCOPE_CYCLE_COUNTER(STAT_Mass_IsSubjectRecursive);
+
+		if (Subject != Object
+			&& Subject.IsValid()
+			&& Object.IsValid())
+		{
+			TArray<FMassEntityHandle> Subjects;
+			TArray<FMassEntityHandle> Objects;
+			Subjects.Add(Subject);
+			
+			constexpr int32 IterationsLimit = 64;
+			for (int32 SubjectIndex = 0; SubjectIndex < Subjects.Num() && ensure(SubjectIndex < IterationsLimit); ++SubjectIndex)
+			{
+				if (const FRelationData::FRelationInstanceData* SubjectData = RelationData.RoleMap.Find(Subjects[SubjectIndex]))
+				{
+					GetRoleEntities((*SubjectData)[static_cast<int32>(ERelationRole::Object)], Objects);
+
+					if (Objects.Find(Object) != INDEX_NONE)
+					{
+						return true;
+					}
+
+					for (const FMassEntityHandle& ObjectHandle : Objects)
+					{
+						Subjects.AddUnique(ObjectHandle);
+					}
+					Objects.Reset();
+				}
+			}
+		}
+
+		return false;
+	}
+
+	void FRelationManager::OnRelationTypeRegistered(const FTypeHandle InRegisteredTypeHandle, const FRelationTypeTraits& RelationTypeTraits)
+	{
+		FRelationData& RelationData = CreateRelationData(InRegisteredTypeHandle);
+		ensureMsgf(RelationData.Traits.bHierarchical == false || RelationData.Traits.RegisteredGroupType.IsValid()
+			, TEXT("Hierarchical relationships need a valid registered group size. Failed by relationship %s")
+			, *RelationData.Traits.GetFName().ToString());
+
+		TArray<TSubclassOf<UMassRelationObserver>> ObserversToCreate = 
+		{
+			RelationTypeTraits.RelationEntityCreationObserverClass.Get()
+			, RelationTypeTraits.RelationEntityDestructionObserverClass.Get()
+			, UMassRelationEntityGuardDog::StaticClass()
+		};
+
+		if (RelationTypeTraits.SubjectEntityDestructionObserverClass.Get() == UMassRelationRoleDestruction::StaticClass()
+			&& RelationTypeTraits.ObjectEntityDestructionObserverClass.Get() == UMassRelationRoleDestruction::StaticClass())
+		{
+			UMassRelationRoleDestruction::AddObserverInstances(EntityManager.GetObserverManager(), InRegisteredTypeHandle, RelationData.Traits);
+		}
+		else
+		{
+			ObserversToCreate.Add(RelationTypeTraits.SubjectEntityDestructionObserverClass.Get());
+			ObserversToCreate.Add(RelationTypeTraits.ObjectEntityDestructionObserverClass.Get());
+		}
+
+		for (const TSubclassOf<UMassRelationObserver>& ObserverClass : ObserversToCreate)
+		{
+			if (ObserverClass)
+			{
+				UMassRelationObserver* ObserverProcessor = NewObject<UMassRelationObserver>(EntityManager.GetOwner(), ObserverClass);
+				if (ObserverProcessor->ConfigureRelationObserver(InRegisteredTypeHandle, RelationData.Traits))
+				{
+					EntityManager.GetObserverManager().AddObserverInstance(ObserverProcessor);
+				}
+			}
+		}
+	}
+} // namespace UE::Mass
+
+#undef WITH_RELATIONSHIP_VALIDATION
diff --git a/Engine/Source/Runtime/MassEntity/Private/MassRelationObservers.cpp b/Engine/Source/Runtime/MassEntity/Private/MassRelationObservers.cpp
new file mode 100644
--- /dev/null
+++ b/Engine/Source/Runtime/MassEntity/Private/MassRelationObservers.cpp
@@ -0,0 +1,318 @@
+// Copyright Epic Games, Inc. All Rights Reserved.
+
+#include "MassRelationObservers.h"
+#include "MassTypeManager.h"
+#include "MassExecutionContext.h"
+#include "MassObserverManager.h"
+
+//-----------------------------------------------------------------------------
+// UMassRelationObserver
+//-----------------------------------------------------------------------------
+UMassRelationObserver::UMassRelationObserver()
+	: EntityQuery(*this)
+{
+	bAutoRegisterWithObserverRegistry = false;
+	ObservedOperations = EMassObservedOperationFlags::None;
+}
+
+void UMassRelationObserver::InitializeInternal(UObject& Owner, const TSharedRef<FMassEntityManager>& EntityManager)
+{
+	Super::InitializeInternal(Owner, EntityManager);
+}
+
+void UMassRelationObserver::ConfigureQueries(const TSharedRef<FMassEntityManager>& EntityManager)
+{
+	check(ObservedType);
+	checkf(UE::Mass::IsA<FMassFragment>(ObservedType) || UE::Mass::IsA<FMassTag>(ObservedType)
+		, TEXT("Only tags and fragments are valid observed types for RelationObservers. Received %s")
+		, *ObservedType->GetName());
+
+	if (UE::Mass::IsA<FMassTag>(ObservedType))
+	{
+		EntityQuery.AddTagRequirement(*ObservedType, EMassFragmentPresence::All);
+	}
+	else // UE::Mass::IsA<FMassFragment>(ObservedType)
+	{
+		EntityQuery.AddRequirement(ObservedType, ObservedTypeAccess);
+	}
+
+	const UE::Mass::FRelationManager& RelationManager = EntityManager->GetRelationManager();
+	const UE::Mass::FRelationData& RelationData = RelationManager.GetRelationDataChecked(RelationTypeHandle);
+	const TNotNull<const UScriptStruct*> RelationTag = RelationData.Traits.GetRelationTagType();
+	if (RelationTag != ObservedType)
+	{
+		EntityQuery.AddTagRequirement(*RelationTag, EMassFragmentPresence::All);
+	}
+
+	if (bAutoAddRelationFragment)
+	{
+		const TNotNull<const UScriptStruct*> RelationFragmentType = RelationData.Traits.RelationFragmentType;
+		if (UE::Mass::IsA<FMassFragment>(RelationFragmentType) && RelationFragmentType != ObservedType)
+		{
+			EntityQuery.AddRequirement(RelationFragmentType, RelationFragmentAccessType, EMassFragmentPresence::All);
+		}
+	}
+
+	DebugDescription = FString::Printf(TEXT("RelationType: %s ObservedType: %s Operation: %s")
+		, *RelationTypeHandle.ToString(), *ObservedType->GetName()
+		, *LexToString(ObservedOperations));
+}
+
+bool UMassRelationObserver::ConfigureRelationObserver(UE::Mass::FTypeHandle InRegisteredTypeHandle, const UE::Mass::FRelationTypeTraits& Traits)
+{
+	ObservedType = &const_cast<UScriptStruct&>(*Traits.RelationFragmentType);
+	RelationTypeHandle = InRegisteredTypeHandle;
+	return true;
+}
+
+//-----------------------------------------------------------------------------
+// UMassRelationEntityCreation
+//-----------------------------------------------------------------------------
+UMassRelationEntityCreation::UMassRelationEntityCreation()
+{
+	ObservedOperations = EMassObservedOperationFlags::CreateEntity;
+	ExecutionPriority = RelationCreationObserverExecutionPriority;
+};
+
+void UMassRelationEntityCreation::Execute(FMassEntityManager& EntityManager, FMassExecutionContext& Context)
+{
+	/*
+	// here's example code showing how the freshly-created relation entity can be utilized.
+	// at the current state of MassRelations implementation we don't need to do anything here. 
+
+	struct FRelationInstanceRegistration
+	{
+		FMassEntityHandle SubjectHandle;
+		FMassEntityHandle ObjectHandle;
+		FMassEntityHandle RelationEntityHandle;
+	};
+
+	TArray<FRelationInstanceRegistration> RelationInstances;
+	EntityQuery.ForEachEntityChunk(Context, [&RelationInstances, ObservedType = ObservedType](FMassExecutionContext& ExecutionContext)
+	{
+		TConstArrayView<FMassRelationFragment> RelationFragments = ExecutionContext.GetFragmentView<FMassRelationFragment>(ObservedType);
+		for (FMassExecutionContext::FEntityIterator EntityIt = ExecutionContext.CreateEntityIterator(); EntityIt; ++EntityIt)
+		{
+			RelationInstances.Add({
+				RelationFragments[EntityIt].Subject
+				, RelationFragments[EntityIt].Object
+				, ExecutionContext.GetEntity(EntityIt)
+			});
+		}
+	});
+	*/
+}
+
+//-----------------------------------------------------------------------------
+// UMassRelationEntityGuardDog
+//-----------------------------------------------------------------------------
+UMassRelationEntityGuardDog::UMassRelationEntityGuardDog()
+{
+	ObservedOperations = EMassObservedOperationFlags::RemoveElement;
+}
+
+void UMassRelationEntityGuardDog::Execute(FMassEntityManager& EntityManager, FMassExecutionContext& InContext)
+{
+#if WITH_MASSENTITY_DEBUG
+	EntityQuery.CacheArchetypes();
+	const FMassArchetypeEntityCollection& EntityCollection = InContext.GetEntityCollection();
+	ensureMsgf(EntityQuery.GetArchetypes().Contains(EntityCollection.GetArchetype()) == false
+		, TEXT("Trying to remove private-implementation-detail fragments from a relation entity is not supported!"));
+#endif // WITH_MASSENTITY_DEBUG
+}
+
+//-----------------------------------------------------------------------------
+// UMassRelationEntityDestruction
+//-----------------------------------------------------------------------------
+UMassRelationEntityDestruction::UMassRelationEntityDestruction()
+{
+	ObservedOperations = EMassObservedOperationFlags::DestroyEntity;
+};
+
+void UMassRelationEntityDestruction::Execute(FMassEntityManager& EntityManager, FMassExecutionContext& Context)
+{
+	TArray<TPair<FMassEntityHandle, FMassEntityHandle>> EntitiesToClearOut;
+
+	EntityQuery.ForEachEntityChunk(Context, [&EntitiesToClearOut, ObservedType = ObservedType](FMassExecutionContext& ExecutionContext)
+	{
+		TConstArrayView<FMassRelationFragment> RelationFragments = ExecutionContext.GetFragmentView<FMassRelationFragment>(ObservedType);
+		for (FMassExecutionContext::FEntityIterator EntityIt = ExecutionContext.CreateEntityIterator(); EntityIt; ++EntityIt)
+		{
+			EntitiesToClearOut.Add({RelationFragments[EntityIt].Subject, RelationFragments[EntityIt].Object});
+		}
+	});
+
+	if (EntitiesToClearOut.Num())
+	{
+		UE::Mass::FRelationManager& RelationManager = Context.GetEntityManagerChecked().GetRelationManager();
+		UE::Mass::FRelationData& RelationData = RelationManager.GetRelationDataChecked(RelationTypeHandle);
+
+		for (const TPair<FMassEntityHandle, FMassEntityHandle>& Pair : EntitiesToClearOut)
+		{
+			if (UE::Mass::FRelationData::FRelationInstanceData* InstanceData = RelationData.RoleMap.Find(Pair.Get<0>()))
+			{
+				(*InstanceData)[static_cast<int32>(UE::Mass::ERelationRole::Object)].RemoveAllSwap(
+					FMassRelationRoleInstanceHandle::FMassRelationRoleInstanceHandleFinder(Pair.Get<1>())
+					, EAllowShrinking::No);
+			}
+			if (UE::Mass::FRelationData::FRelationInstanceData* InstanceData = RelationData.RoleMap.Find(Pair.Get<1>()))
+			{
+				(*InstanceData)[static_cast<int32>(UE::Mass::ERelationRole::Subject)].RemoveAllSwap(
+					FMassRelationRoleInstanceHandle::FMassRelationRoleInstanceHandleFinder(Pair.Get<0>())
+					, EAllowShrinking::No);
+			}
+		}
+	}
+}
+
+//-----------------------------------------------------------------------------
+// UMassRelationRoleDestruction
+//-----------------------------------------------------------------------------
+UMassRelationRoleDestruction::UMassRelationRoleDestruction()
+{
+	ObservedOperations = EMassObservedOperationFlags::DestroyEntity;
+	bAutoAddRelationFragment = false;
+};
+
+void UMassRelationRoleDestruction::Execute(FMassEntityManager& EntityManager, FMassExecutionContext& Context)
+{
+	EntityQuery.ForEachEntityChunk(Context, ExecuteFunction);
+}
+
+void UMassRelationRoleDestruction::ConfigureQueries(const TSharedRef<FMassEntityManager>& EntityManager)
+{
+	Super::ConfigureQueries(EntityManager);
+	// extend the query with "none" reliance on the relation entity's data fragment
+	// @todo missing feature - this would cause the query to always fail to find archetypes. We need a way to
+	//		add "external requirements" to queries, that we'd use to calculate dependencies, but not use for binding
+
+	if (ExcludedRelationFragmentType)
+	{
+		EntityQuery.AddRequirement(ExcludedRelationFragmentType, EMassFragmentAccess::None, EMassFragmentPresence::None);
+	}
+}
+
+bool UMassRelationRoleDestruction::ConfigureRelationObserver(UE::Mass::FTypeHandle InRegisteredTypeHandle, const UE::Mass::FRelationTypeTraits& Traits)
+{
+	check(RelationRole != UE::Mass::ERelationRole::MAX && static_cast<uint8>(UE::Mass::ERelationRole::MAX) == 2);
+
+	const int32 RoleAsIndex = static_cast<int32>(RelationRole);
+	const int32 OppositeRoleIndex = (RoleAsIndex + 1) % 2;
+	bool bExecutionFunctionAssigned = false;
+
+	// note there we're deliberately not calling Super::ConfigureRelationObserver, since we're going to observe role-specific elements
+	RelationTypeHandle = InRegisteredTypeHandle;
+	if (Traits.RoleTraits[RoleAsIndex].Element)
+	{
+		ObservedType = Traits.RoleTraits[RoleAsIndex].Element;
+	}
+	else
+	{
+		TNotNull<const UScriptStruct*> RelationTypeTag = RelationTypeHandle.GetScriptStruct();
+		ObservedType = RelationTypeTag;
+		ExcludedRelationFragmentType = Traits.RelationFragmentType;
+	}
+
+	switch (Traits.RoleTraits[RoleAsIndex].DestructionPolicy)
+	{
+		case UE::Mass::ERemovalPolicy::CleanUp:
+			RelationFragmentAccessType = EMassFragmentAccess::None;
+			ExecuteFunction = [InRegisteredTypeHandle, OppositeRoleIndex](FMassExecutionContext& ExecutionContext)
+			{
+				UE::Mass::FRelationManager& RelationManager = ExecutionContext.GetEntityManagerChecked().GetRelationManager();
+				UE::Mass::FRelationData& RelationData = RelationManager.GetRelationDataChecked(InRegisteredTypeHandle);
+				for (FMassExecutionContext::FEntityIterator EntityIt = ExecutionContext.CreateEntityIterator(); EntityIt; ++EntityIt)
+				{
+					TArray<FMassRelationRoleInstanceHandle>& Container = RelationData.RoleMap.FindChecked(ExecutionContext.GetEntity(EntityIt))[OppositeRoleIndex];
+					ExecutionContext.Defer().DestroyEntities(RelationManager.GetRelationEntities(Container));
+					Container.Reset();
+				}
+			};
+			bExecutionFunctionAssigned = true;
+			break;
+		case UE::Mass::ERemovalPolicy::Destroy:
+			RelationFragmentAccessType = EMassFragmentAccess::ReadOnly;
+			// destroy the other side of the relationship, and the relation entity
+			ExecuteFunction = [InRegisteredTypeHandle, RoleAsIndex = RoleAsIndex](FMassExecutionContext& ExecutionContext)
+			{
+				const int32 OppositeRoleIndex = (RoleAsIndex + 1) % 2;
+				UE::Mass::FRelationManager& RelationManager = ExecutionContext.GetEntityManagerChecked().GetRelationManager();
+				UE::Mass::FRelationData& RelationData = RelationManager.GetRelationDataChecked(InRegisteredTypeHandle);
+
+				TArray<FMassEntityHandle> EntitiesToDestroy;
+
+				for (FMassExecutionContext::FEntityIterator EntityIt = ExecutionContext.CreateEntityIterator(); EntityIt; ++EntityIt)
+				{
+					TArray<FMassRelationRoleInstanceHandle>& Container = RelationData.RoleMap.FindChecked(ExecutionContext.GetEntity(EntityIt))[OppositeRoleIndex];
+					RelationManager.GetRelationEntities(Container, EntitiesToDestroy);
+					RelationManager.GetRoleEntities(Container, EntitiesToDestroy);
+				}
+
+				ExecutionContext.Defer().DestroyEntities(EntitiesToDestroy);
+			};
+			bExecutionFunctionAssigned = true;
+			break;
+		case UE::Mass::ERemovalPolicy::Splice:
+			// use my Object as my Subject's Object, and vive versa
+			ExecuteFunction = [InRegisteredTypeHandle](FMassExecutionContext& ExecutionContext)
+			{
+				static auto DestroyRelationEntitiesAndGetRoleEntities = [](
+					const FMassExecutionContext& LocalExecutionContext
+					, const UE::Mass::FRelationManager& RelationManager
+					, const UE::Mass::FRelationData::FRelationInstanceData& InstanceData
+					, const UE::Mass::ERelationRole Role) -> TArray<FMassEntityHandle>
+				{
+					const TArray<FMassRelationRoleInstanceHandle>& RoleRelationInstanceHandles = InstanceData[Role];
+
+					TArray<FMassEntityHandle> RoleRelationEntityHandles = RelationManager.GetRelationEntities(RoleRelationInstanceHandles);
+					LocalExecutionContext.Defer().DestroyEntities(MoveTemp(RoleRelationEntityHandles));
+
+					return RelationManager.GetRoleEntities(RoleRelationInstanceHandles);
+				};
+
+				UE::Mass::FRelationManager& RelationManager = ExecutionContext.GetEntityManagerChecked().GetRelationManager();
+				UE::Mass::FRelationData& RelationData = RelationManager.GetRelationDataChecked(InRegisteredTypeHandle);
+
+				for (FMassExecutionContext::FEntityIterator EntityIt = ExecutionContext.CreateEntityIterator(); EntityIt; ++EntityIt)
+				{
+					const FMassEntityHandle RoleEntity = ExecutionContext.GetEntity(EntityIt);
+					if (UE::Mass::FRelationData::FRelationInstanceData* InstanceData = RelationData.RoleMap.Find(RoleEntity))
+					{
+						// @todo optimize array allocations here - we could create the arrays outside and re-use
+						TArray<FMassEntityHandle> SubjectEntities = DestroyRelationEntitiesAndGetRoleEntities(ExecutionContext, RelationManager, *InstanceData, UE::Mass::ERelationRole::Subject);
+						TArray<FMassEntityHandle> ObjectEntities = DestroyRelationEntitiesAndGetRoleEntities(ExecutionContext, RelationManager, *InstanceData, UE::Mass::ERelationRole::Object);
+						
+						RelationManager.CreateRelationInstances(InRegisteredTypeHandle, MakeArrayView(SubjectEntities), MakeArrayView(ObjectEntities));
+					}
+				}
+			};
+			bExecutionFunctionAssigned = true;
+			break;
+		case UE::Mass::ERemovalPolicy::Custom: // fall through on purse
+		default: ;
+	}
+
+	return bExecutionFunctionAssigned;
+}
+
+void UMassRelationRoleDestruction::AddObserverInstances(FMassObserverManager& ObserverManager, UE::Mass::FTypeHandle InRegisteredTypeHandle, const UE::Mass::FRelationTypeTraits& Traits)
+{
+	bool NeedsSameObserver[static_cast<int32>(UE::Mass::ERelationRole::MAX)] = {
+		false
+		, (Traits.RoleTraits[1].Element == Traits.RoleTraits[0].Element && Traits.RoleTraits[1].DestructionPolicy == Traits.RoleTraits[0].DestructionPolicy)
+	};
+
+	for (int32 RoleIndex = 0; RoleIndex < static_cast<int32>(UE::Mass::ERelationRole::MAX); ++RoleIndex)
+	{
+		if (Traits.RoleTraits[RoleIndex].DestructionPolicy != UE::Mass::ERemovalPolicy::Custom
+		   && !NeedsSameObserver[RoleIndex])
+		{
+			UMassRelationRoleDestruction* ObserverProcessor = NewObject<UMassRelationRoleDestruction>();
+			ObserverProcessor->RelationRole = static_cast<UE::Mass::ERelationRole>(RoleIndex);
+			if (ObserverProcessor->ConfigureRelationObserver(InRegisteredTypeHandle, Traits))
+			{
+				ObserverManager.AddObserverInstance(ObserverProcessor);
+			}
+		}
+	}
+}
diff --git a/Engine/Source/Runtime/MassEntity/Private/MassTypeManager.cpp b/Engine/Source/Runtime/MassEntity/Private/MassTypeManager.cpp
--- a/Engine/Source/Runtime/MassEntity/Private/MassTypeManager.cpp
+++ b/Engine/Source/Runtime/MassEntity/Private/MassTypeManager.cpp
@@ -4,6 +4,8 @@
 #include "MassEntityManager.h"
 #include "MassTestableEnsures.h"
 #include "MassEntityElementTypes.h"
+#include "MassEntityRelations.h"
+#include "Misc/CoreDelegates.h"
 
 
 namespace UE::Mass
@@ -20,11 +22,19 @@ namespace UE::Mass
 	//-----------------------------------------------------------------------------
 	// FTypeManager
 	//-----------------------------------------------------------------------------
+	FTypeManager::FOnRegisterBuiltInTypes FTypeManager::OnRegisterBuiltInTypes;
+
 	FTypeManager::FTypeManager(FMassEntityManager& InEntityManager)
 		: OuterEntityManager(InEntityManager)
 	{
 	}
 
+	void FTypeManager::RegisterBuiltInTypes()
+	{
+		OnRegisterBuiltInTypes.Broadcast(*this);
+		bBuiltInTypesRegistered = true;
+	}
+
 	FTypeHandle FTypeManager::RegisterTypeInternal(TNotNull<const UStruct*> InType, FTypeInfo&& TypeInfo)
 	{
 		FTypeHandle TypeHandle(InType);
@@ -38,7 +48,7 @@ namespace UE::Mass
 
 			TypeDataMap.Add(FTypeHandle(TypeHandle.TypeKey), MoveTemp(TypeInfo));
 		}
-		else
+		else if (testableEnsureMsgf(bBuiltInTypesRegistered == false, TEXT("Registered type overriding is supported only as part of built-in types registration.")))
 		{
 			// we're overriding the existing data with the new data in assumption it's more up-to-date.
 			// The most common occurence of this will be with already registered subsystems' subclasses.
@@ -68,4 +78,67 @@ namespace UE::Mass
 		TypeInfo.Traits.Set<FSharedFragmentTypeTraits>(MoveTemp(TypeTraits));
 		return RegisterTypeInternal(InType, MoveTemp(TypeInfo));
 	}
+
+	FTypeHandle FTypeManager::RegisterType(FRelationTypeTraits&& TypeTraits)
+	{
+		TNotNull<const UScriptStruct*> InType = TypeTraits.GetRelationTagType();
+		testableCheckfReturn(UE::Mass::IsA<FMassRelation>(InType), return {}
+			, TEXT("%s is not a valid relation type"), *InType->GetName());
+
+		testableCheckfReturn(TypeTraits.RelationFragmentType->IsChildOf(FMassRelationFragment::StaticStruct()), return {}
+			, TEXT("%s is not a valid TypeTraits.RelationFragmentType, needs to derive from FMassRelationFragment")
+			, *TypeTraits.RelationFragmentType->GetName());
+
+		if (TypeTraits.RelationName.IsNone())
+		{
+			TypeTraits.RelationName = InType->GetFName();
+			TypeTraits.SetDebugInFix(TypeTraits.RelationName.ToString());
+		}
+
+		FTypeInfo* ExistingData = TypeDataMap.Find(FTypeHandle(InType));
+		
+		if (!testableEnsureMsgf(bBuiltInTypesRegistered == false || ExistingData == nullptr
+			, TEXT("Modifying relationship after registration done is not supported")))
+		{
+			return {};
+		}
+
+		FTypeInfo TypeInfo;
+		TypeInfo.TypeName = InType->GetFName();
+		if (TypeTraits.RegisteredGroupType.IsValid() == false)
+		{
+			TypeTraits.RegisteredGroupType = OuterEntityManager.FindOrAddArchetypeGroupType(TypeTraits.RelationName);
+		}
+		TypeInfo.Traits.Set<FRelationTypeTraits>(MoveTemp(TypeTraits));
+
+		const FTypeHandle RegisteredTypeHandle = RegisterTypeInternal(InType, MoveTemp(TypeInfo));
+		if (RegisteredTypeHandle.IsValid())
+		{
+			if (bBuiltInTypesRegistered)
+			{
+				// if the built-in types are already registered we need to notify the entity manager that there's a new
+				// relation type, that might require additional handling (like creation of appropriate entity destruction observers)
+				// Note that we don't call this during built-in types registration to give project-specific code a chance
+				// to override type traits before the entity manager handles the registered types.
+				OuterEntityManager.OnNewTypeRegistered(RegisteredTypeHandle);
+			}
+		}
+
+		return RegisteredTypeHandle;
+	}
+
+	FTypeHandle FTypeManager::GetRelationTypeHandle(const TNotNull<const UScriptStruct*> RelationOrElementType) const
+	{
+		const bool bValidRelationshipType = UE::Mass::IsA<FMassRelation>(RelationOrElementType);
+		ensureMsgf(bValidRelationshipType, TEXT("%s is not a valid relationship type"), *RelationOrElementType->GetName());
+
+		return UE::Mass::IsA<FMassRelation>(RelationOrElementType) 
+			? FTypeHandle(RelationOrElementType)
+			: FTypeHandle();
+	}
+
+	bool FTypeManager::IsValidRelationType(const TNotNull<const UScriptStruct*> RelationOrElementType) const
+	{
+		return UE::Mass::IsA<FMassRelation>(RelationOrElementType) && TypeDataMap.Contains(FTypeHandle(RelationOrElementType));
+	}
 }
\ No newline at end of file
diff --git a/Engine/Source/Runtime/MassEntity/Private/Relations/MassChildOf.cpp b/Engine/Source/Runtime/MassEntity/Private/Relations/MassChildOf.cpp
new file mode 100644
--- /dev/null
+++ b/Engine/Source/Runtime/MassEntity/Private/Relations/MassChildOf.cpp
@@ -0,0 +1,99 @@
+// Copyright Epic Games, Inc. All Rights Reserved.
+
+#include "Relations/MassChildOf.h"
+#include "MassTypeManager.h"
+#include "MassEntityManager.h"
+#include "Misc/CoreDelegates.h"
+#include "MassExecutionContext.h"
+
+namespace UE::Mass
+{
+	namespace Relations
+	{
+		FTypeHandle ChildOfHandle;
+
+		struct FChildOfRelationTypeInitializer
+		{
+			FChildOfRelationTypeInitializer()
+				: RegisterBuiltInTypesHandle(FTypeManager::OnRegisterBuiltInTypes.AddStatic(&FChildOfRelationTypeInitializer::RegisterType))
+			{
+				UE::Mass::Relations::ChildOfHandle = FTypeManager::MakeTypeHandle(FMassChildOfRelation::StaticStruct());
+			}
+
+			~FChildOfRelationTypeInitializer()
+			{
+				FTypeManager::OnRegisterBuiltInTypes.Remove(RegisterBuiltInTypesHandle);
+			}
+
+			static void RegisterType(FTypeManager& InOutTypeManager)
+			{
+				FRelationTypeTraits ChildOfRelation(FMassChildOfRelation::StaticStruct());
+				ChildOfRelation.RoleTraits[static_cast<int32>(ERelationRole::Object)].bExclusive = true;
+				ChildOfRelation.RoleTraits[static_cast<int32>(ERelationRole::Object)].DestructionPolicy = ERemovalPolicy::Destroy;
+
+				ChildOfRelation.RoleTraits[static_cast<int32>(ERelationRole::Subject)].bExclusive = false;
+				ChildOfRelation.RoleTraits[static_cast<int32>(ERelationRole::Subject)].DestructionPolicy = ERemovalPolicy::CleanUp;
+				// @todo unused for now.
+				ChildOfRelation.RoleTraits[static_cast<int32>(ERelationRole::Subject)].Element = FMassChildOfFragment::StaticStruct();
+
+				ChildOfRelation.bHierarchical = true;
+				ChildOfRelation.RelationName = "ChildOf";
+				
+				ChildOfRelation.RegisteredGroupType = InOutTypeManager.GetEntityManager().FindOrAddArchetypeGroupType(ChildOfRelation.RelationName);
+
+				ChildOfRelation.RelationEntityCreationObserverClass = UMassChildOfRelationEntityCreation::StaticClass();
+
+				InOutTypeManager.RegisterType(MoveTemp(ChildOfRelation));
+			}
+
+			const FDelegateHandle RegisterBuiltInTypesHandle;
+		};
+
+		void RegisterChildOfRelation()
+		{
+			static FChildOfRelationTypeInitializer BuildInRelationHandlesInitializer;
+		}
+	}
+}
+
+//-----------------------------------------------------------------------------
+// UMassRelationEntityCreation
+//-----------------------------------------------------------------------------
+void UMassChildOfRelationEntityCreation::Execute(FMassEntityManager& EntityManager, FMassExecutionContext& Context)
+{
+	struct FRelationInstanceRegistration
+	{
+		FMassEntityHandle ChildHandle;
+		FMassEntityHandle ParentHandle;
+		FMassEntityHandle RelationEntityHandle;
+	};
+
+	TArray<FRelationInstanceRegistration> RelationInstances;
+	EntityQuery.ForEachEntityChunk(Context, [&RelationInstances, ObservedType = ObservedType](FMassExecutionContext& ExecutionContext)
+	{
+		TConstArrayView<FMassRelationFragment> RelationFragments = ExecutionContext.GetFragmentView<FMassRelationFragment>(ObservedType);
+		for (FMassExecutionContext::FEntityIterator EntityIt = ExecutionContext.CreateEntityIterator(); EntityIt; ++EntityIt)
+		{
+			RelationInstances.Add({
+				RelationFragments[EntityIt].Subject
+				, RelationFragments[EntityIt].Object
+				, ExecutionContext.GetEntity(EntityIt)
+			});
+		}
+	});
+	// @todo if we don't do mapping how do we "remove relation", how do we find the relation entity?
+	// @todo, again, we need "external requirement" meaning "things we touch on other entities, so needs
+	// to be accounted for in dependencies, but we're not binding or expecting them to be in the very archetypes we process"
+	for (const FRelationInstanceRegistration& RelationInstance : RelationInstances)
+	{
+		// using GetFragmentDataStruct to support this observer processor being used by relations that extend the ChildOf
+		// relation by using a fragment deriving from FMassChildOfFragment
+		// @todo use the fragment type from traits, not hardcoded FMassChildOfFragment
+		FStructView FragmentView = EntityManager.GetFragmentDataStruct(RelationInstance.ChildHandle, FMassChildOfFragment::StaticStruct());
+		FMassChildOfFragment* ChildOfFragment = FragmentView.GetPtr<FMassChildOfFragment>();
+		if (ensure(ChildOfFragment))
+		{
+			ChildOfFragment->Parent = RelationInstance.ParentHandle;
+		}
+	}
+}
diff --git a/Engine/Source/Runtime/MassEntity/Public/MassEntityBuilder.h b/Engine/Source/Runtime/MassEntity/Public/MassEntityBuilder.h
--- a/Engine/Source/Runtime/MassEntity/Public/MassEntityBuilder.h
+++ b/Engine/Source/Runtime/MassEntity/Public/MassEntityBuilder.h
@@ -6,6 +6,8 @@
 #include "MassEntityConcepts.h"
 #include "MassEntityHandle.h"
 #include "MassEntityTypes.h"
+#include "MassEntityRelations.h"
+#include "MassTypeManager.h"
 
 #define UE_API MASSENTITY_API
 
@@ -13,6 +15,7 @@ struct FMassEntityManager;
 
 namespace UE::Mass
 {
+	struct FTypeHandle;
 	namespace Private
 	{
 		struct FEntityBuilderHelper;
@@ -180,6 +183,33 @@ struct FEntityBuilder
 	/** Adds the ElementType to the target archetype's composition */
 	UE_API FEntityBuilder& Add(TNotNull<const UScriptStruct*> ElementType);
 
+	/** type used to store parameters for relations to be created once the target entity is created. */
+	struct FPendingRelationParams
+	{
+		UE::Mass::FTypeHandle RelationTypeHandle;
+		FMassEntityHandle OtherEntity;
+		Relations::ERelationRole OtherEntityRole;
+	};
+
+	/**
+	 * Adds information about a specific relation instance to be added once the entity gets created.
+	 * The function will simply store the information without any checks for validity or duplication.
+	 * If you want to override existing relation data call ForEachRelation.
+	 */
+	UE_API FEntityBuilder& AddRelation(UE::Mass::FTypeHandle RelationTypeHandle, FMassEntityHandle OtherEntity, Relations::ERelationRole InputEntityRole = Relations::ERelationRole::Object);
+
+	/** templated helper function for calling the other AddRelation function */
+	template<UE::Mass::CRelation T>
+	FEntityBuilder& AddRelation(FMassEntityHandle OtherEntity, Relations::ERelationRole InputEntityRole = Relations::ERelationRole::Object);
+
+	/**
+	 * Calls the provided Operator function for every stored pending relation data instance.
+	 * The return value of the Operator is used to determine whether the relation data instance
+	 * should be kept (meaning: return `false` for each element you want to remove).
+	 * The potential element removal is stable.
+	 */
+	UE_API void ForEachRelation(const TFunctionRef<bool(FPendingRelationParams&)>& Operator);
+
 	/**
 	 * Finds and retrieves a pointer to a fragment of type T if it exists.
 	 * @param T - The type of fragment to find.
@@ -281,9 +311,13 @@ private:
 
 	void CacheEntityHandle() const;
 
+	static void CreateEntityImpl(FMassEntityManager& EntityManager, FMassEntityHandle EntityHandle, const FMassArchetypeHandle& ArchetypeHandle
+		, const FMassArchetypeSharedFragmentValues& SharedFragmentValues, TConstArrayView<FInstancedStruct> Fragments, TConstArrayView<FPendingRelationParams> RelationsParams);
+
 	TArray<FInstancedStruct> Fragments;
 	TArray<FInstancedStruct> SharedFragments;
 	TArray<FInstancedStruct> ConstSharedFragments;
+	TArray<FPendingRelationParams> RelationsParams;
 
 	enum class EState : uint8
 	{
@@ -401,6 +435,12 @@ T* FEntityBuilder::Find() requires (CElement<T> && !(CTag<T> || CChunkFragment<T
 	return nullptr;
 }
 
+template<UE::Mass::CRelation T>
+FEntityBuilder& FEntityBuilder::AddRelation(const FMassEntityHandle OtherEntity, const Relations::ERelationRole InputEntityRole)
+{
+	return AddRelation(FTypeManager::MakeTypeHandle(T::StaticStruct()), OtherEntity, InputEntityRole);
+}
+
 inline bool FEntityBuilder::IsValid() const
 {
 	return State != EState::Invalid;
diff --git a/Engine/Source/Runtime/MassEntity/Public/MassEntityManager.h b/Engine/Source/Runtime/MassEntity/Public/MassEntityManager.h
--- a/Engine/Source/Runtime/MassEntity/Public/MassEntityManager.h
+++ b/Engine/Source/Runtime/MassEntity/Public/MassEntityManager.h
@@ -21,6 +21,8 @@
 #include "Templates/FunctionFwd.h"
 #include "MassEntityManagerStorage.h"
 #include "MassArchetypeGroup.h"
+#include "MassRelationManager.h"
+#include "MassEntityRelations.h"
 
 #if WITH_MASSENTITY_DEBUG
 #include "Misc/EnumClassFlags.h"
@@ -46,6 +48,7 @@ namespace UE::Mass
 {
 	struct FEntityBuilder;
 	struct FTypeManager;
+	struct FTypeHandle;
 	namespace Private
 	{
 		struct FEntityStorageInitializer;
@@ -1037,6 +1040,24 @@ PRAGMA_ENABLE_DEPRECATION_WARNINGS
 	template<UE::Mass::CConstSharedFragment T>
 	void ForEachConstSharedFragmentConditional(TFunctionRef< bool(const T& /*ConstSharedFragment*/) > ConditionFunction, TFunctionRef< void(const T& /*ConstSharedFragment*/) > ExecuteFunction);
 
+	//-----------------------------------------------------------------------------
+	// Relations
+	//-----------------------------------------------------------------------------
+	UE::Mass::FRelationManager& GetRelationManager();
+
+	template<UE::Mass::CRelation T>
+	bool BatchCreateRelations(TArrayView<FMassEntityHandle> Subjects, TArrayView<FMassEntityHandle> Objects);
+	UE_API bool BatchCreateRelations(const UE::Mass::FTypeHandle RelationTypeHandle, TArrayView<FMassEntityHandle> Subjects, TArrayView<FMassEntityHandle> Objects);
+
+	//-----------------------------------------------------------------------------
+	// Type management
+	//-----------------------------------------------------------------------------
+	void OnNewTypeRegistered(UE::Mass::FTypeHandle RegisteredTypeHandle);
+
+protected:
+	void OnRelationTypeRegistered(UE::Mass::FTypeHandle RegisteredTypeHandle, const UE::Mass::FRelationTypeTraits& RelationTypeTraits);
+
+public:
 	//-----------------------------------------------------------------------------
 	// Entity Builder
 	//-----------------------------------------------------------------------------
@@ -1286,6 +1307,7 @@ private:
 	FMassObserverManager ObserverManager;
 
 	TSharedRef<UE::Mass::FTypeManager> TypeManager;
+	UE::Mass::FRelationManager RelationManager;
 
 	TMap<const FName, const int32> GroupNameToTypeIndex;
 	// @todo we'll probably have some "GroupTypeInformation" here in the future
@@ -1438,4 +1460,15 @@ inline TSharedRef<FMassObserverManager::FObserverLock> FMassEntityManager::GetOr
 	return ObserverManager.GetOrMakeObserverLock();
 }
 
-#undef UE_API 
\ No newline at end of file
+inline UE::Mass::FRelationManager& FMassEntityManager::GetRelationManager()
+{
+	return RelationManager;
+}
+
+template<UE::Mass::CRelation T>
+bool FMassEntityManager::BatchCreateRelations(TArrayView<FMassEntityHandle> Subjects, TArrayView<FMassEntityHandle> Objects)
+{
+	return BatchCreateRelations(UE::Mass::FTypeManager::MakeTypeHandle(T::StaticStruct()), Subjects, Objects);
+}
+
+#undef UE_API
diff --git a/Engine/Source/Runtime/MassEntity/Public/MassEntityRelations.h b/Engine/Source/Runtime/MassEntity/Public/MassEntityRelations.h
new file mode 100644
--- /dev/null
+++ b/Engine/Source/Runtime/MassEntity/Public/MassEntityRelations.h
@@ -0,0 +1,342 @@
+// Copyright Epic Games, Inc. All Rights Reserved.
+
+#pragma once
+
+#include "MassEntityElementTypes.h"
+#include "MassEntityHandle.h"
+#include "MassArchetypeGroup.h"
+#include "MassEntityMacros.h"
+#include "MassEntityRelations.generated.h"
+
+#define UE_API MASSENTITY_API
+
+struct FMassEntityManager;
+class UMassObserverProcessor;
+namespace UE::Mass
+{
+	struct FRelationManager;
+}
+
+namespace UE::Mass::Relations
+{
+	enum class ERemovalPolicy : uint8
+	{
+		/** only update the book-keeping */
+		CleanUp,
+		/** when a given relation gets removed destroy the source entity (eg destroy the child when ChildOf relation gets removed) */
+		Destroy,
+		/** external code will be called to patch up what remains off the relation */
+		Splice,
+		/** the user will provide the observer processor */
+		Custom,
+
+		MAX,
+		Default = CleanUp,
+	};
+
+	enum class EExternalMappingRequired : uint8
+	{
+		No,
+		Yes,
+		Default = Yes, 
+	};
+
+	enum class ERelationRole : uint8
+	{
+		Subject,
+		Object,
+		MAX
+	};
+
+	inline FString LexToString(const ERelationRole Value)
+	{
+		return Value == ERelationRole::MAX
+			? TEXT("INVALID")
+			: (Value == ERelationRole::Subject
+				? TEXT("Subject")
+				: TEXT("Object"));
+	}
+
+	struct FRoleTraits
+	{
+		/** 
+		 * The element type that will be added to the participating entity. Can remain empty. 
+		 * Every entity will get the relation tag as well, the FRelationTypeTraits.RelationTagType
+		 */
+		const UScriptStruct* Element = nullptr;
+
+		/**
+		 * What to do when entities serving this role are destroyed. This value affects the whole relation, including other participants. 
+		 */
+		ERemovalPolicy DestructionPolicy = ERemovalPolicy::Default;
+
+		/**
+		 * "Exclusive" means there can be only one participant like this in a relation.
+		 * For example, there can be only one parent in our ChildOf relation (bExclusive = true, on the "parent" participant, i.e. the `Object`)
+		 * while there can be multiple children (bExclusive = false for the `Subject`)
+		 */
+		bool bExclusive = true;
+
+		/**
+		 * Declares whether this relation-specific implementation details provide dedicated
+		 * mechanism for mapping other roles for entities serving as this role, the role these traits affect.
+		 * It essentially answers whether this role needs the system to find other relation participants (`Yes`)
+		 * or it can handle the task with custom code (like via the Element, option `No`).
+		 * 
+		 * Defaulting to `Yes` to provide reliable, out-of-the-box functionality required for
+		 * fetching participants of relation instances. Set to `No` if you want to save memory by
+		 * not automatically populating and utilizing FRelationData.RoleMap.
+		 *
+		 * @todo this functionality is not plugged in yet.
+		 */
+		EExternalMappingRequired RequiresExternalMapping = EExternalMappingRequired::Default;
+
+		bool operator==(const FRoleTraits& Other) const
+		{
+			return Element == Other.Element
+				&& DestructionPolicy == Other.DestructionPolicy
+				&& bExclusive == Other.bExclusive
+				&& RequiresExternalMapping == Other.RequiresExternalMapping;
+		}
+	};
+
+	struct FRelationTypeTraits
+	{
+		UE_API FRelationTypeTraits(TNotNull<const UScriptStruct*> InRelationTagType);
+		UE_API FRelationTypeTraits(const FRelationTypeTraits& Other, TNotNull<const UScriptStruct*> NewRelationTagType);
+		UE_API FRelationTypeTraits(const FRelationTypeTraits& Other) = default;
+
+		TNotNull<const UScriptStruct*> GetRelationTagType() const
+		{
+			return RelationTagType;
+		}
+
+		FName GetFName() const
+		{
+			return RelationName;
+		}
+
+		FName RelationName;
+
+		/** Checks whether the traits are configured properly. */
+		bool IsValid() const;
+
+	private:
+		/** can only be set during creation since we use the same type to create UE::Mass::FTypeHandle for the relation type */
+		TNotNull<const UScriptStruct*> RelationTagType;
+
+	public:
+		/**
+		 * The fragment type that will be automatically added to each relation entity created to represent relation instances
+		 * @todo we can switch this to an array of fragment types or FInstancedStructs, bitsets or even FMassArchetypeCompositionDescriptor,
+		 *		but first we need to see how our use-cases shape up.
+		 */
+		TNotNull<const UScriptStruct*> RelationFragmentType;
+
+		/**
+		 * Whether to use hierarchical archetype groups for the relation entities
+		 * Set this to true of you want to do any data calculations with data stored in
+		 * the relation entities, and them being processed in hierarchy order being a requirements
+		 */
+		bool bCreateRelationEntitiesInHierarchy = false;
+
+		TStaticArray<FRoleTraits, static_cast<uint8>(ERelationRole::MAX)> RoleTraits;
+
+		bool bHierarchical = false;
+
+		FArchetypeGroupType RegisteredGroupType;
+
+		/**
+		 * Is set, gets called upon type's registration to register appropriate observers.
+		 * Return value indicates whether the entity manager should register default observes as well.
+		 */
+		TFunction<bool(FMassEntityManager&)> RegisterObserversDelegate;
+
+		/**
+		 * Processor classes to be instantiated when the relation type gets registered.
+		 * These will get auto-created only if RegisterObserversDelegate is empty or returns `true`
+		 * The default values point to generic observer classes that implement the generic policies and behavior.
+		 * Nulling-out any of these member variables will result in the given functionality not being implemented
+		 * by the system and makes the user responsible for supplying it. 
+		 */
+		TWeakObjectPtr<UClass> RelationEntityCreationObserverClass;
+		TWeakObjectPtr<UClass> RelationEntityDestructionObserverClass;
+		TWeakObjectPtr<UClass> SubjectEntityDestructionObserverClass;
+		TWeakObjectPtr<UClass> ObjectEntityDestructionObserverClass;
+
+		void SetDebugInFix(FString&& InFix);
+#if WITH_MASSENTITY_DEBUG
+		UE_API FString DebugDescribeRelation(FMassEntityHandle A, FMassEntityHandle B) const;
+	private:
+		FString DebugInFix;
+#endif // WITH_MASSENTITY_DEBUG
+	};
+}
+
+//-----------------------------------------------------------------------------
+// Relation types
+//-----------------------------------------------------------------------------
+/**
+ * Structs extending FMassRelation represent a "concept" or a relation. These structs are
+ * not intended to be stored in Mass. @see FMassRelationFragment for ways of storing relation-specific data
+ */
+USTRUCT()
+struct FMassRelation : public FMassTag
+{
+	GENERATED_BODY()
+};
+
+//-----------------------------------------------------------------------------
+// Relation data
+//-----------------------------------------------------------------------------
+/**
+ * Relation fragment base. Every relation entity will get an instance of this type
+ * or a type derived from it (as configured via FRelationTypeTraits.RelationFragmentType
+ */
+USTRUCT()
+struct FMassRelationFragment : public FMassFragment
+{
+	GENERATED_BODY()
+
+	/**
+	 * This is the "who" part of the relation. Examples:
+	 * - the Child in a ChildOf relation
+	 * - the Character in a HasWeapon relation
+	 * - the Employee in a WorksFor relation
+	 */
+	FMassEntityHandle Subject;
+
+	/**
+	 * This is the "what" or "target" part of the relation. Examples:
+	 * - the Parent in a ChildOf relation
+	 * - the Weapon in a HasWeapon relation
+	 * - the Company in a WorksFor relation
+	 */
+	FMassEntityHandle Object;
+
+	FMassEntityHandle GetRole(const int32 Index) const
+	{
+		return Index == 0 ? Subject : Object;
+	}
+
+	FMassEntityHandle GetRole(const UE::Mass::Relations::ERelationRole Role) const
+	{
+		check(Role != UE::Mass::Relations::ERelationRole::MAX);
+		return GetRole(static_cast<int32>(Role));
+	}
+};
+
+USTRUCT()
+struct FMassRelationMappingFragment : public FMassFragment
+{
+	GENERATED_BODY()
+};
+
+/**
+ * @todo we might want to promote the internals to FEntityIdentifier and have FMassEntityHandle derive from that.
+ */
+struct alignas(8) FMassRelationRoleInstanceHandle
+{
+	static constexpr int32 EntityIndexBits = 30;
+	static constexpr int32 EntityIndexMask = (1 << 30) - 1;
+	static constexpr int32 TypeMask = ~EntityIndexMask;
+
+	FMassRelationRoleInstanceHandle() = default;
+
+	static FMassRelationRoleInstanceHandle Create(UE::Mass::Relations::ERelationRole Role, const FMassEntityHandle RoleHandle, const FMassEntityHandle RelationEntityHandle)
+	{
+		check(Role != UE::Mass::Relations::ERelationRole::MAX);
+
+		FMassRelationRoleInstanceHandle ReturnHandle;
+		ReturnHandle.SetRoleEntityIndex(RoleHandle.Index);
+		ReturnHandle.SetRelationEntityIndex(RelationEntityHandle.Index);
+		ReturnHandle.SetRole(Role);
+		check(ReturnHandle.GetRoleEntityIndex() == RoleHandle.Index)
+		check(ReturnHandle.GetRelationEntityIndex() == RelationEntityHandle.Index)
+
+		return ReturnHandle;
+	}
+
+	int32 GetRoleEntityIndex() const 
+	{
+		return (RoleEntity & EntityIndexMask);
+	}
+	MASSENTITY_API FMassEntityHandle GetRoleEntityHandle(const FMassEntityManager& EntityManager) const;
+
+	int32 GetRelationEntityIndex() const 
+	{
+		return (RelationEntity & EntityIndexMask);
+	}
+	MASSENTITY_API FMassEntityHandle GetRelationEntityHandle(const FMassEntityManager& EntityManager) const;
+
+	UE::Mass::Relations::ERelationRole GetRole() const
+	{
+		return static_cast<UE::Mass::Relations::ERelationRole>((RoleEntity & TypeMask) >> EntityIndexBits);
+	}
+
+	friend FString LexToString(const FMassRelationRoleInstanceHandle Handle)
+	{
+		return Handle.DebugGetDescription();
+	}
+
+	FString DebugGetDescription() const
+	{
+		return FString::Printf(TEXT("Relation %d role %s %d"), GetRelationEntityIndex(), *LexToString(GetRole()), GetRoleEntityIndex());
+	}
+
+	/** We're unable to tell if a given relation instance handle is valid just by looking at a handle. Only the RelationManager can answer this question. Use IsSet as first filter */
+	bool IsValid() const = delete;
+
+	struct FMassRelationRoleInstanceHandleFinder
+	{
+		FMassRelationRoleInstanceHandleFinder(const FMassEntityHandle EntityHandle)
+			: EntityIndex(EntityHandle.Index)
+		{	
+		}
+
+		bool operator()(const FMassRelationRoleInstanceHandle& Element) const
+		{
+			return Element.GetRoleEntityIndex() == EntityIndex;
+		}
+
+		const int32 EntityIndex;
+	};
+
+	bool operator==(const FMassRelationRoleInstanceHandle& Other) const
+	{
+		return RelationEntity == Other.RelationEntity && RoleEntity == Other.RoleEntity;
+	}
+
+private:
+	uint32 RelationEntity = 0;
+	uint32 RoleEntity = 0;
+
+	void SetRoleEntityIndex(const int32 InIndex)
+	{
+		RoleEntity = (InIndex & EntityIndexMask) | (RoleEntity & TypeMask);
+	}
+
+	void SetRelationEntityIndex(const int32 InIndex)
+	{
+		RelationEntity = (InIndex & EntityIndexMask) | (RelationEntity & TypeMask);
+	}
+
+	void SetRole(UE::Mass::Relations::ERelationRole InRole)
+	{
+		RoleEntity = (RoleEntity & EntityIndexMask) | (static_cast<int32>(InRole) << EntityIndexBits);
+	}
+};
+
+namespace UE::Mass
+{
+	template<>
+	inline bool IsA<FMassRelation>(const UStruct* Struct)
+	{
+		return Struct && Struct->IsChildOf(FMassRelation::StaticStruct());
+	}
+
+	template<typename T>
+	concept CRelation = TIsDerivedFrom<typename TRemoveReference<T>::Type, FMassRelation>::Value;
+};
+
+#undef UE_API
\ No newline at end of file
diff --git a/Engine/Source/Runtime/MassEntity/Public/MassRelationCommands.h b/Engine/Source/Runtime/MassEntity/Public/MassRelationCommands.h
new file mode 100644
--- /dev/null
+++ b/Engine/Source/Runtime/MassEntity/Public/MassRelationCommands.h
@@ -0,0 +1,77 @@
+// Copyright Epic Games, Inc. All Rights Reserved.
+
+#pragma once
+
+#include "MassCommands.h"
+#include "MassEntityManager.h"
+
+#if CSV_PROFILER_STATS || WITH_MASSENTITY_DEBUG
+#	define DEBUG_NAME(Name) , FName(TEXT(Name))
+#	define DEBUG_NAME_PARAM(Name) , const FName InDebugName = TEXT(Name)
+#	define FORWARD_DEBUG_NAME_PARAM , InDebugName
+#else
+#	define DEBUG_NAME(Name)
+#	define DEBUG_NAME_PARAM(Name)
+#	define FORWARD_DEBUG_NAME_PARAM
+#endif // CSV_PROFILER_STATS || WITH_MASSENTITY_DEBUG
+
+template<typename TRelation>
+struct FMassCommandMakeRelation : public FMassBatchedEntityCommand
+{
+	using Super = FMassBatchedEntityCommand;
+
+	FMassCommandMakeRelation(EMassCommandOperationType OperationType = EMassCommandOperationType::Add DEBUG_NAME_PARAM("MakeRelation"))
+		: Super(OperationType FORWARD_DEBUG_NAME_PARAM)
+	{}
+
+	void Add(const FMassEntityHandle ChildEntity, const FMassEntityHandle ParentEntity)
+	{
+		Super::Add(ChildEntity);
+		Parents.Add(ParentEntity);
+	}
+
+	void Add(TConstArrayView<FMassEntityHandle> ChildEntities, TConstArrayView<FMassEntityHandle> ParentEntities)
+	{
+		ensure(ChildEntities.Num() != 0 && ParentEntities.Num() != 0);
+
+		Super::Add(ChildEntities);
+
+		if (ChildEntities.Num() > ParentEntities.Num())
+		{
+			// @todo to be improved - we should have a dedicated path for Multi-children -> Single parent  operations
+			do
+			{
+				Parents.Append(ParentEntities.GetData(), ParentEntities.Num());
+			} while (Parents.Num() < TargetEntities.Num());
+			Parents.SetNum(TargetEntities.Num(), EAllowShrinking::No);
+		}
+		else
+		{
+			Parents.Append(ParentEntities.GetData(), FMath::Min(ParentEntities.Num(), ChildEntities.Num()));
+		}
+	}
+
+protected:
+	virtual void Reset() override
+	{
+		Parents.Reset();
+		Super::Reset();
+	}
+
+	virtual SIZE_T GetAllocatedSize() const override
+	{
+		return Super::GetAllocatedSize() + Parents.GetAllocatedSize();
+	}
+
+	virtual void Run(FMassEntityManager& EntityManager) override
+	{
+		TRACE_CPUPROFILER_EVENT_SCOPE(FMassCommandMakeRelation_Execute);
+		EntityManager.BatchCreateRelations<TRelation>(TargetEntities, Parents);
+	}
+
+	TArray<FMassEntityHandle> Parents;
+};
+
+#undef DEBUG_NAME
+#undef DEBUG_NAME_PARAM
+#undef FORWARD_DEBUG_NAME_PARAM
diff --git a/Engine/Source/Runtime/MassEntity/Public/MassRelationManager.h b/Engine/Source/Runtime/MassEntity/Public/MassRelationManager.h
new file mode 100644
--- /dev/null
+++ b/Engine/Source/Runtime/MassEntity/Public/MassRelationManager.h
@@ -0,0 +1,230 @@
+// Copyright Epic Games, Inc. All Rights Reserved.
+
+#pragma once
+
+#include "UObject/Class.h"
+#include "MassTypeManager.h"
+#include "MassEntityHandle.h"
+#include "MassEntityRelations.h"
+
+#define UE_API MASSENTITY_API
+
+struct FMassEntityHandle;
+struct FMassEntityManager;
+struct FMassExecutionContext;
+
+namespace UE::Mass
+{
+	struct FTypeRegistry;
+	struct FRelationManager;
+	
+	struct FRelationData
+	{
+		FRelationData(const FRelationTypeTraits& InTraits);
+
+		const FRelationTypeTraits Traits;
+
+		struct FRelationInstanceData
+		{
+			TStaticArray<TArray<FMassRelationRoleInstanceHandle>, static_cast<uint32>(ERelationRole::MAX)> RelationEntityContainers;
+
+			bool IsEmpty() const
+			{
+				bool bIsEmpty = true;
+				for (const TArray<FMassRelationRoleInstanceHandle>& Container : RelationEntityContainers)
+				{
+					bIsEmpty = bIsEmpty && Container.IsEmpty();
+				}
+				return bIsEmpty;
+			}
+
+			TArray<FMassRelationRoleInstanceHandle>& operator[](const int32 Index)
+			{
+				return RelationEntityContainers[Index];
+			}
+
+			const TArray<FMassRelationRoleInstanceHandle>& operator[](const int32 Index) const
+			{
+				return RelationEntityContainers[Index];
+			}
+
+			TArray<FMassRelationRoleInstanceHandle>& operator[](const ERelationRole Index)
+			{
+				return (*this)[static_cast<int32>(Index)];
+			}
+
+			const TArray<FMassRelationRoleInstanceHandle>& operator[](const ERelationRole Index) const
+			{
+				return (*this)[static_cast<int32>(Index)];
+			}
+		};
+		TMap<FMassEntityHandle, FRelationInstanceData> RoleMap;
+
+		TArray<FMassEntityHandle> GetParticipants(const FMassEntityManager& EntityManager, const FMassEntityHandle RoleEntity, ERelationRole QueriedRole) const;
+	};
+
+
+	struct FRelationManager
+	{
+		UE_API explicit FRelationManager(FMassEntityManager& EntityManager);
+
+		template<UE::Mass::CRelation T>
+		FMassEntityHandle CreateRelationInstance(FMassEntityHandle Subject, FMassEntityHandle Object);
+		UE_API FMassEntityHandle CreateRelationInstance(const FTypeHandle RelationTypeHandle, FMassEntityHandle Subject, FMassEntityHandle Object);
+
+		/** Creates a relation type handle with RelationType, and calls the other CreateRelationInstances implementation */
+		UE_API TArray<FMassEntityHandle> CreateRelationInstances(TNotNull<const UScriptStruct*> RelationType, TArrayView<FMassEntityHandle> Subjects, TArrayView<FMassEntityHandle> Objects);
+
+		/**
+		 * Creates valid relation instances of type RelationTypeHandle, binding Subjects and Objects
+		 * Note that the input arrays can have their order modified by the function, all the relation pairs that are not valid, are moved to the back of the arrays
+		 * The number of elements in Subjects and Objects must match.
+		 */
+		UE_API TArray<FMassEntityHandle> CreateRelationInstances(const FTypeHandle RelationTypeHandle, TArrayView<FMassEntityHandle> Subjects, TArrayView<FMassEntityHandle> Objects);
+
+		UE_API bool DestroyRelationInstance(const FTypeHandle RelationTypeHandle, const FMassEntityHandle Subject, const FMassEntityHandle Object);
+		UE_API bool DestroyRelationInstance(FMassRelationRoleInstanceHandle RelationHandle) const;
+
+		/** Fetch all the entities that are "subjects" in instances of the given relation type, where ObjectEntity is the "object" of the relation */
+		[[nodiscard]] TArray<FMassEntityHandle> GetRelationSubjects(TNotNull<const UScriptStruct*> RelationType, const FMassEntityHandle ObjectEntity) const;
+		[[nodiscard]] TArray<FMassEntityHandle> GetRelationSubjects(const FTypeHandle RelationTypeHandle, const FMassEntityHandle ObjectEntity) const;
+
+		/** Fetch all the entities that are "objects" in instances of the given relation type, where SubjectEntity is the "subject" of the relation */
+		[[nodiscard]] TArray<FMassEntityHandle> GetRelationObjects(TNotNull<const UScriptStruct*> RelationType, const FMassEntityHandle SubjectEntity) const;
+		[[nodiscard]] TArray<FMassEntityHandle> GetRelationObjects(const FTypeHandle RelationTypeHandle, const FMassEntityHandle SubjectEntity) const;
+
+		[[nodiscard]] TArray<FMassEntityHandle> GetRelationEntities(TConstArrayView<FMassRelationRoleInstanceHandle> RelationEntitiesContainer) const;
+		void GetRelationEntities(TConstArrayView<FMassRelationRoleInstanceHandle> RelationEntitiesContainer, TArray<FMassEntityHandle>& InOutEntityHandles) const;
+		[[nodiscard]] TArray<FMassEntityHandle> GetRoleEntities(TConstArrayView<FMassRelationRoleInstanceHandle> RelationEntitiesContainer) const;
+		void GetRoleEntities(TConstArrayView<FMassRelationRoleInstanceHandle> RelationEntitiesContainer, TArray<FMassEntityHandle>& InOutEntityHandles) const;
+
+		FTypeHandle GetRelationTypeHandle(TNotNull<const UScriptStruct*> RelationType) const;
+
+		UE_API bool IsSubjectOfRelation(const FTypeHandle RelationTypeHandle, const FMassEntityHandle Subject, const FMassEntityHandle Object) const;
+		UE_API bool IsSubjectOfRelation(const FRelationData& RelationDataInstance, const FMassEntityHandle Subject, const FMassEntityHandle Object) const;
+		UE_API bool IsSubjectOfRelationRecursive(const FTypeHandle RelationTypeHandle, const FMassEntityHandle Subject, const FMassEntityHandle Object) const;
+		UE_API bool IsSubjectOfRelationRecursive(const FRelationData& RelationDataInstance, const FMassEntityHandle Subject, const FMassEntityHandle Object) const;
+
+
+		void OnRelationTypeRegistered(const FTypeHandle RegisteredTypeHandle, const FRelationTypeTraits& RelationTypeTraits);
+
+		const FRelationData& GetRelationDataChecked(const FTypeHandle RelationTypeHandle) const;
+		FRelationData& GetRelationDataChecked(const FTypeHandle RelationTypeHandle);
+
+	protected:
+
+		FRelationData& CreateRelationData(const FTypeHandle RelationTypeHandle);
+
+		UE_API TArray<FMassEntityHandle> GetRelationObjects(const FRelationData& RelationData, const FMassEntityHandle SubjectEntity) const;
+		UE_API TArray<FMassEntityHandle> GetRelationSubjects(const FRelationData& RelationData, const FMassEntityHandle ObjectEntity) const;
+
+		const FRelationData* GetRelationData(const FTypeHandle RelationTypeHandle) const;
+
+		struct FHierarchyEntitiesContainer
+		{
+			void StoreUnique(const uint32 Depth, TArray<FMassEntityHandle>& InOutArray);
+			void StoreUnique(const uint32 Depth, FMassEntityHandle Handle);
+
+			TArray<FMassEntityHandle>& operator[](const uint32 Depth)
+			{
+				return ContainerPerLevel[Depth];
+			}
+			const TArray<FMassEntityHandle>& operator[](const uint32 Depth) const
+			{
+				return ContainerPerLevel[Depth];
+			}
+
+			uint32 Num() const
+			{
+				return static_cast<uint32>(ContainerPerLevel.Num());
+			}
+		protected:
+			TArray<TArray<FMassEntityHandle>> ContainerPerLevel;
+			TSet<FMassEntityHandle> ExistingElements;
+		};
+
+		void GatherHierarchy(const FRelationData& RelationData, const FMassEntityHandle SubjectHandle, FHierarchyEntitiesContainer& SubSubjects, uint32 Depth = 0) const;
+		
+		FMassEntityManager& EntityManager;
+		const FTypeManager& TypeManager;
+
+		TMap<FTypeHandle, FRelationData> RelationsDataMap;
+	};
+
+	//-----------------------------------------------------------------------------
+	// INLINES
+	//-----------------------------------------------------------------------------
+	template<UE::Mass::CRelation T>
+	FMassEntityHandle FRelationManager::CreateRelationInstance(FMassEntityHandle Subject, FMassEntityHandle Object)
+	{
+		TArray<FMassEntityHandle> CreatedRelationshipEntities = CreateRelationInstances(T::StaticStruct(), TArrayView<FMassEntityHandle>(&Subject, 1), TArrayView<FMassEntityHandle>(&Object, 1));
+		return CreatedRelationshipEntities.IsEmpty() ? FMassEntityHandle() : CreatedRelationshipEntities[0];
+	}
+
+	inline FMassEntityHandle FRelationManager::CreateRelationInstance(const FTypeHandle RelationTypeHandle, FMassEntityHandle Subject, FMassEntityHandle Object)
+	{
+		TArray<FMassEntityHandle> CreatedRelationshipEntities = CreateRelationInstances(RelationTypeHandle, TArrayView<FMassEntityHandle>(&Subject, 1), TArrayView<FMassEntityHandle>(&Object, 1));
+		return CreatedRelationshipEntities.IsEmpty() ? FMassEntityHandle() : CreatedRelationshipEntities[0];
+	}
+
+	inline TArray<FMassEntityHandle> FRelationManager::CreateRelationInstances(TNotNull<const UScriptStruct*> RelationType, TArrayView<FMassEntityHandle> Subjects, TArrayView<FMassEntityHandle> Objects)
+	{
+		const FTypeHandle RelationTypeHandle = TypeManager.GetRelationTypeHandle(RelationType);
+		if (!ensureMsgf(RelationTypeHandle.IsValid(), TEXT("%hs: Unknown relation type %s, make sure to register it first"), __FUNCTION__, *RelationType->GetName()))
+		{
+			return {};
+		}
+
+		return CreateRelationInstances(RelationTypeHandle, Subjects, Objects);
+	}
+
+	inline FTypeHandle FRelationManager::GetRelationTypeHandle(TNotNull<const UScriptStruct*> RelationType) const
+	{
+		checkf(UE::Mass::IsA<FMassRelation>(RelationType)
+			, TEXT("Provided RelationType, %s, is not a relation type")
+			, *RelationType->GetName());
+
+		return TypeManager.GetRelationTypeHandle(RelationType);
+	}
+	
+	inline const FRelationData* FRelationManager::GetRelationData(const FTypeHandle RelationTypeHandle) const
+	{
+		return RelationsDataMap.Find(RelationTypeHandle);
+	}
+
+	inline const FRelationData& FRelationManager::GetRelationDataChecked(const FTypeHandle RelationTypeHandle) const
+	{
+		return RelationsDataMap.FindChecked(RelationTypeHandle);
+	}
+
+	inline FRelationData& FRelationManager::GetRelationDataChecked(const FTypeHandle RelationTypeHandle)
+	{
+		return RelationsDataMap.FindChecked(RelationTypeHandle);
+	}
+
+	inline TArray<FMassEntityHandle> FRelationManager::GetRelationSubjects(TNotNull<const UScriptStruct*> RelationType, const FMassEntityHandle ObjectEntity) const
+	{
+		const FTypeHandle RelationTypeHandle = TypeManager.GetRelationTypeHandle(RelationType);
+		return GetRelationSubjects(RelationTypeHandle, ObjectEntity);
+	}
+
+	inline TArray<FMassEntityHandle> FRelationManager::GetRelationSubjects(const FTypeHandle RelationTypeHandle, const FMassEntityHandle ObjectEntity) const
+	{
+		const FRelationData& RelationData = GetRelationDataChecked(RelationTypeHandle);
+		return GetRelationSubjects(RelationData, ObjectEntity);
+	}
+
+	inline TArray<FMassEntityHandle> FRelationManager::GetRelationObjects(TNotNull<const UScriptStruct*> RelationType, const FMassEntityHandle SubjectEntity) const
+	{
+		const FTypeHandle RelationTypeHandle = TypeManager.GetRelationTypeHandle(RelationType);
+		return GetRelationObjects(RelationTypeHandle, SubjectEntity);
+	}
+
+	inline TArray<FMassEntityHandle> FRelationManager::GetRelationObjects(const FTypeHandle RelationTypeHandle, const FMassEntityHandle SubjectEntity) const
+	{
+		const FRelationData& RelationData = GetRelationDataChecked(RelationTypeHandle);
+		return GetRelationObjects(RelationData, SubjectEntity);
+	}
+} // namespace UE::Mass
+
+#undef UE_API
\ No newline at end of file
diff --git a/Engine/Source/Runtime/MassEntity/Public/MassRelationObservers.h b/Engine/Source/Runtime/MassEntity/Public/MassRelationObservers.h
new file mode 100644
--- /dev/null
+++ b/Engine/Source/Runtime/MassEntity/Public/MassRelationObservers.h
@@ -0,0 +1,97 @@
+// Copyright Epic Games, Inc. All Rights Reserved.
+
+#pragma once
+
+#include "MassObserverProcessor.h"
+#include "MassEntityQuery.h"
+#include "MassRelationObservers.generated.h"
+
+namespace UE::Mass::Relations
+{
+	struct FRelationTypeTraits;
+	enum class ERelationRole : uint8;
+}
+
+UCLASS()
+class MASSENTITY_API UMassRelationObserver : public UMassObserverProcessor
+{
+	GENERATED_BODY()
+
+public:
+	UMassRelationObserver();
+
+	/**
+	 * @return whether the configuration was successful. Failed observers won't get added to the observer manager and will get destroyed promptly.
+	 */
+	virtual bool ConfigureRelationObserver(UE::Mass::FTypeHandle InRegisteredTypeHandle, const UE::Mass::FRelationTypeTraits& Traits);
+
+protected:
+	virtual void InitializeInternal(UObject& Owner, const TSharedRef<FMassEntityManager>& EntityManager) override;
+	virtual void ConfigureQueries(const TSharedRef<FMassEntityManager>& EntityManager) override;
+
+	FMassEntityQuery EntityQuery;
+	UE::Mass::FTypeHandle RelationTypeHandle;
+	/** Relevant only when ObservedType is a fragment */
+	EMassFragmentAccess ObservedTypeAccess = EMassFragmentAccess::ReadWrite;
+	EMassFragmentAccess RelationFragmentAccessType = EMassFragmentAccess::ReadWrite;
+	bool bAutoAddRelationFragment = true;
+
+	FString DebugDescription;
+};
+
+UCLASS()
+class MASSENTITY_API UMassRelationEntityCreation : public UMassRelationObserver
+{
+	GENERATED_BODY()
+public:
+	static constexpr int32 RelationCreationObserverExecutionPriority = 1024;
+
+	UMassRelationEntityCreation();
+
+protected:
+	virtual void Execute(FMassEntityManager& EntityManager, FMassExecutionContext& Context) override;
+};
+
+/** Debug-time, detects relation elements being removed from the relation entities, which is not supposed to be done */
+UCLASS()
+class MASSENTITY_API UMassRelationEntityGuardDog : public UMassRelationObserver
+{
+	GENERATED_BODY()
+public:
+	UMassRelationEntityGuardDog();
+
+protected:
+	virtual void Execute(FMassEntityManager& EntityManager, FMassExecutionContext& Context) override;
+};
+
+/** cleans up data, removes enties in RoleMap if Entry.IsEmpty */
+UCLASS()
+class MASSENTITY_API UMassRelationEntityDestruction : public UMassRelationObserver
+{
+	GENERATED_BODY()
+public:
+	UMassRelationEntityDestruction();
+
+protected:
+	virtual void Execute(FMassEntityManager& EntityManager, FMassExecutionContext& Context) override;
+};
+
+UCLASS()
+class MASSENTITY_API UMassRelationRoleDestruction : public UMassRelationObserver
+{
+	GENERATED_BODY()
+
+public:
+	UMassRelationRoleDestruction();
+
+	virtual bool ConfigureRelationObserver(UE::Mass::FTypeHandle InRegisteredTypeHandle, const UE::Mass::FRelationTypeTraits& Traits) override;
+	static void AddObserverInstances(FMassObserverManager& GetObserverManager, UE::Mass::FTypeHandle InRegisteredTypeHandle, const UE::Mass::FRelationTypeTraits& Traits);
+
+protected:
+	virtual void Execute(FMassEntityManager& EntityManager, FMassExecutionContext& Context) override;
+	virtual void ConfigureQueries(const TSharedRef<FMassEntityManager>& EntityManager) override;
+
+	FMassExecuteFunction ExecuteFunction;
+	UE::Mass::Relations::ERelationRole RelationRole = UE::Mass::Relations::ERelationRole::MAX;
+	const UScriptStruct* ExcludedRelationFragmentType = nullptr;
+};
diff --git a/Engine/Source/Runtime/MassEntity/Public/MassTypeManager.h b/Engine/Source/Runtime/MassEntity/Public/MassTypeManager.h
--- a/Engine/Source/Runtime/MassEntity/Public/MassTypeManager.h
+++ b/Engine/Source/Runtime/MassEntity/Public/MassTypeManager.h
@@ -9,6 +9,7 @@
 #include "Templates/SharedPointer.h"
 #include "MassEntityConcepts.h"
 #include "Subsystems/Subsystem.h"
+#include "MassEntityRelations.h"
 #include "MassExternalSubsystemTraits.h"
 
 #define UE_API MASSENTITY_API
@@ -18,10 +19,14 @@ struct FMassEntityManager;
 
 namespace UE::Mass
 {
+	using namespace Relations;
+	struct FRelationData;
+
 	/** Handle for identifying and managing types in the type manager */
 	struct FTypeHandle
 	{
 		FTypeHandle() = default;
+		FTypeHandle(const FTypeHandle&) = default;
 
 		bool operator==(const FTypeHandle&) const = default;
 
@@ -42,9 +47,21 @@ namespace UE::Mass
 			return GetTypeHash(InHandle.TypeKey);
 		}
 
+		FName GetFName() const
+		{
+			const UStruct* RepresentedType = TypeKey.ResolveObjectPtr();
+			return GetFNameSafe(RepresentedType);
+		}
+
+		FString ToString() const
+		{
+			const UStruct* RepresentedType = TypeKey.ResolveObjectPtr();
+			return GetNameSafe(RepresentedType);
+		}
+
 	private:
 		friend struct FTypeManager;
-		explicit FTypeHandle(TObjectKey<const UStruct> InTypeKey);
+		UE_API explicit FTypeHandle(TObjectKey<const UStruct> InTypeKey);
 
 		TObjectKey<const UStruct> TypeKey;
 	};
@@ -99,7 +116,10 @@ namespace UE::Mass
 	 */
 	struct FTypeInfo
 	{
-		using FTypeTraits = TVariant<FEmptyTypeTraits, FSubsystemTypeTraits, FSharedFragmentTypeTraits>;
+		using FTypeTraits = TVariant<FEmptyTypeTraits
+			, FSubsystemTypeTraits
+			, FSharedFragmentTypeTraits
+			, FRelationTypeTraits>;
 
 		FName TypeName;
 		FTypeTraits Traits;
@@ -109,12 +129,22 @@ namespace UE::Mass
 
 		/** Fetches stored data as FSharedFragmentTypeTraits, if applicable, or null otherwise */
 		const FSharedFragmentTypeTraits* GetAsSharedFragmentTraits() const;
+
+		/** Fetches stored data as FRelationTypeTraits, if applicable, or null otherwise */
+		const FRelationTypeTraits* GetAsRelationTraits() const;
+
+		/** Fetches stored data as FRelationTypeTraits. Will complain if stored data is not of FRelationTypeTraits type */
+		const FRelationTypeTraits& GetAsRelationTraitsChecked() const;
 	};
 
 	struct FTypeManager : TSharedFromThis<FTypeManager>
 	{
+		DECLARE_MULTICAST_DELEGATE_OneParam(FOnRegisterBuiltInTypes, FTypeManager& /*Instance*/);
+
 		UE_API explicit FTypeManager(FMassEntityManager& InEntityManager);
 
+		void RegisterBuiltInTypes();
+
 		/**
 		 * @return whether the type manager instance has any types registered at all.
 		 */
@@ -124,6 +154,8 @@ namespace UE::Mass
 		UE_API FTypeHandle RegisterType(TNotNull<const UStruct*> InType, FSubsystemTypeTraits&&);
 		/** Register traits for given shared fragments type */
 		UE_API FTypeHandle RegisterType(TNotNull<const UStruct*> InType, FSharedFragmentTypeTraits&&);
+		/** Register traits for given relation type */
+		UE_API FTypeHandle RegisterType(FRelationTypeTraits&&);
 
 		/** Registration helper for shared fragments */
 		template<CSharedFragment T>
@@ -143,6 +175,11 @@ namespace UE::Mass
 		 */
 		const FTypeInfo* GetTypeInfo(TObjectKey<const UStruct> TypeKey) const;
 
+		const FRelationTypeTraits& GetRelationTypeChecked(const FTypeHandle TypeHandle) const;
+		const FRelationTypeTraits& GetRelationTypeChecked(const TNotNull<const UScriptStruct*> RelationOrElementType) const;
+		UE_API FTypeHandle GetRelationTypeHandle(const TNotNull<const UScriptStruct*> RelationOrElementType) const;
+		UE_API bool IsValidRelationType(const TNotNull<const UScriptStruct*> RelationOrElementType) const;
+
 		/** Alias for the iterator type that can be used to iterate all of stored type traits information */
 		using FTypeInfoConstIterator = TMap<FTypeHandle, FTypeInfo>::TConstIterator;
 
@@ -158,6 +195,14 @@ namespace UE::Mass
 
 		FMassEntityManager& GetEntityManager();
 
+		static FTypeHandle MakeTypeHandle(TNotNull<const UStruct*> InTypeKey);
+
+		/**
+		 * Broadcasts as part of FTypeManager::RegisterBuiltInTypes call, giving a chance to the external code to
+		 * register additional types that are supposed to be available from the very start.
+		 */
+		UE_API static FOnRegisterBuiltInTypes OnRegisterBuiltInTypes;
+
 	private:
 		/** Register traits for given type */
 		UE_API FTypeHandle RegisterTypeInternal(TNotNull<const UStruct*> InType, FTypeInfo&&);
@@ -169,6 +214,8 @@ namespace UE::Mass
 
 		/** Contains all registered subsystem types. Can be used to filter access to TypeDataMap */
 		TSet<FTypeHandle> SubsystemTypes;
+
+		bool bBuiltInTypesRegistered = false;
 	};
 
 	//----------------------------------------------------------------------------
@@ -192,6 +239,16 @@ namespace UE::Mass
 		return Traits.TryGet<FSharedFragmentTypeTraits>();
 	}
 
+	inline const FRelationTypeTraits* FTypeInfo::GetAsRelationTraits() const
+	{
+		return Traits.TryGet<FRelationTypeTraits>();
+	}
+
+	inline const FRelationTypeTraits& FTypeInfo::GetAsRelationTraitsChecked() const
+	{
+		return Traits.Get<FRelationTypeTraits>();
+	}
+
 	inline bool FTypeManager::IsEmpty() const
 	{
 		return TypeDataMap.IsEmpty();
@@ -230,10 +287,27 @@ namespace UE::Mass
 		return TypeDataMap.Find(FTypeHandle(TypeKey));
 	}
 
+	inline const FRelationTypeTraits& FTypeManager::GetRelationTypeChecked(const FTypeHandle RelationHandle) const
+	{
+		const FTypeInfo* RelationInfo = GetTypeInfo(RelationHandle);
+		check(RelationInfo);
+		return RelationInfo->GetAsRelationTraitsChecked();
+	}
+
+	inline const FRelationTypeTraits& FTypeManager::GetRelationTypeChecked(const TNotNull<const UScriptStruct*> RelationOrElementType) const
+	{
+		return GetRelationTypeChecked(GetRelationTypeHandle(RelationOrElementType));
+	}
+
 	inline FMassEntityManager& FTypeManager::GetEntityManager()
 	{
 		return OuterEntityManager;
 	}
+
+	inline FTypeHandle FTypeManager::MakeTypeHandle(TNotNull<const UStruct*> InTypeKey)
+	{
+		return FTypeHandle(InTypeKey);
+	}
 }
 
 #undef UE_API
diff --git a/Engine/Source/Runtime/MassEntity/Public/Relations/MassChildOf.h b/Engine/Source/Runtime/MassEntity/Public/Relations/MassChildOf.h
new file mode 100644
--- /dev/null
+++ b/Engine/Source/Runtime/MassEntity/Public/Relations/MassChildOf.h
@@ -0,0 +1,42 @@
+// Copyright Epic Games, Inc. All Rights Reserved.
+
+#pragma once
+
+#include "MassEntityRelations.h"
+#include "MassTypeManager.h"
+#include "MassRelationObservers.h"
+#include "MassChildOf.generated.h"
+
+#define UE_API MASSENTITY_API
+ 
+namespace UE::Mass::Relations
+{
+	extern UE_API FTypeHandle ChildOfHandle;
+	UE_API void RegisterChildOfRelation();
+}
+
+USTRUCT()
+struct FMassChildOfRelation : public FMassRelation
+{
+	GENERATED_BODY()
+};
+
+USTRUCT()
+struct FMassChildOfFragment : public FMassFragment
+{
+	GENERATED_BODY()
+
+	UPROPERTY(VisibleAnywhere, Category = Mass)
+	FMassEntityHandle Parent;
+};
+
+UCLASS()
+class UE_API UMassChildOfRelationEntityCreation : public UMassRelationEntityCreation
+{
+	GENERATED_BODY()
+
+protected:
+	virtual void Execute(FMassEntityManager& EntityManager, FMassExecutionContext& Context) override;
+};
+
+#undef UE_API
\ No newline at end of file
