Fix/silence V530 unhandled return value warnings

#rnx

diff --git a/Engine/Plugins/AI/HTNPlanner/Source/HTNPlanner/Private/HTNBuilder.cpp b/Engine/Plugins/AI/HTNPlanner/Source/HTNPlanner/Private/HTNBuilder.cpp
--- a/Engine/Plugins/AI/HTNPlanner/Source/HTNPlanner/Private/HTNBuilder.cpp
+++ b/Engine/Plugins/AI/HTNPlanner/Source/HTNPlanner/Private/HTNBuilder.cpp
@@ -251,7 +251,7 @@ void FHTNBuilder_Domain::Decompile()
 		else
 		{
 			//const FHTNCompositeTask& CompositeTask = Domain.GetCompositeTask(It->Value);
-			AddCompositeTask(It->Key);
+			AddCompositeTask(It->Key); //-V530
 		}
 	}
 
diff --git a/Engine/Plugins/Enterprise/GLTFExporter/Source/GLTFExporter/Public/Json/GLTFJsonIndexedObjectArray.h b/Engine/Plugins/Enterprise/GLTFExporter/Source/GLTFExporter/Public/Json/GLTFJsonIndexedObjectArray.h
--- a/Engine/Plugins/Enterprise/GLTFExporter/Source/GLTFExporter/Public/Json/GLTFJsonIndexedObjectArray.h
+++ b/Engine/Plugins/Enterprise/GLTFExporter/Source/GLTFExporter/Public/Json/GLTFJsonIndexedObjectArray.h
@@ -92,7 +92,7 @@ struct TGLTFJsonIndexedObjectArray : IGLTFJsonArray
 	{
 		TUniquePtr<ElementType> TempPtr(Element);
 		const bool Result = Array.Contains(TempPtr);
-		TempPtr.Release();
+		(void)TempPtr.Release();
 		return Result;
 	}
 
@@ -100,7 +100,7 @@ struct TGLTFJsonIndexedObjectArray : IGLTFJsonArray
 	{
 		TUniquePtr<ElementType> TempPtr(Element);
 		const SizeType Result = Array.Find(TempPtr);
-		TempPtr.Release();
+		(void)TempPtr.Release();
 		return Result;
 	}
 
diff --git a/Engine/Plugins/NetcodeUnitTest/NetcodeUnitTest/Source/NetcodeUnitTest/Private/ClientUnitTest.cpp b/Engine/Plugins/NetcodeUnitTest/NetcodeUnitTest/Source/NetcodeUnitTest/Private/ClientUnitTest.cpp
--- a/Engine/Plugins/NetcodeUnitTest/NetcodeUnitTest/Source/NetcodeUnitTest/Private/ClientUnitTest.cpp
+++ b/Engine/Plugins/NetcodeUnitTest/NetcodeUnitTest/Source/NetcodeUnitTest/Private/ClientUnitTest.cpp
@@ -602,7 +602,7 @@ void UClientUnitTest::NotifySuspendRequest()
 				if (ResumePipe.IsReadyForRW())
 				{
 					int32 ResumeVal = 1;
-					ResumePipe.WriteInt32(ResumeVal);
+					ResumePipe.WriteInt32(ResumeVal); //-V530
 
 					UNIT_LOG(, TEXT("Sent resume request to server."));
 
diff --git a/Engine/Source/Developer/MassEntityTestSuite/Private/MassSharedFragmentsTest.cpp b/Engine/Source/Developer/MassEntityTestSuite/Private/MassSharedFragmentsTest.cpp
--- a/Engine/Source/Developer/MassEntityTestSuite/Private/MassSharedFragmentsTest.cpp
+++ b/Engine/Source/Developer/MassEntityTestSuite/Private/MassSharedFragmentsTest.cpp
@@ -274,7 +274,7 @@ struct FSharedFragment_ForEach : FExecutionTestBase
 		
 		for (int32 InitValue : TestInitValues)
 		{
-			EntityManager->GetOrCreateSharedFragment<FTestSharedFragment_Int>(InitValue);
+			EntityManager->GetOrCreateSharedFragment<FTestSharedFragment_Int>(InitValue); //-V530
 		}
 
 		TArray<int32> Results;
@@ -334,7 +334,7 @@ struct FConstSharedFragment_ForEach : FExecutionTestBase
 		
 		for (int32 InitValue : TestInitValues)
 		{
-			EntityManager->GetOrCreateConstSharedFragment<FTestConstSharedFragment_Int>(InitValue);
+			EntityManager->GetOrCreateConstSharedFragment<FTestConstSharedFragment_Int>(InitValue); //-V530
 		}
 
 		TArray<int32> Results;
diff --git a/Engine/Source/Runtime/Core/Public/Stats/StatsSystemTypes.h b/Engine/Source/Runtime/Core/Public/Stats/StatsSystemTypes.h
--- a/Engine/Source/Runtime/Core/Public/Stats/StatsSystemTypes.h
+++ b/Engine/Source/Runtime/Core/Public/Stats/StatsSystemTypes.h
@@ -1880,7 +1880,7 @@ struct FThreadSafeStaticStat : public FThreadSafeStaticStatInner<TStatData, TSta
 		//Those macros are guarded for the stats collection to be active which prevented the registration of the stat group.
 		//It was not possible to activate the stat group unless another was already active.
 		//Most groups are registered when a FScopeCycleCounter is declared as GetStatId is called as the constructor parameter.
-		FThreadSafeStaticStatInner<TStatData, TStatData::TGroup::CompileTimeEnable>::GetStatId();
+		FThreadSafeStaticStatInner<TStatData, TStatData::TGroup::CompileTimeEnable>::GetStatId(); //-V530
 	}
 };
 
diff --git a/Engine/Source/Runtime/Core/Tests/Async/TaskGraphTest.cpp b/Engine/Source/Runtime/Core/Tests/Async/TaskGraphTest.cpp
--- a/Engine/Source/Runtime/Core/Tests/Async/TaskGraphTest.cpp
+++ b/Engine/Source/Runtime/Core/Tests/Async/TaskGraphTest.cpp
@@ -969,7 +969,7 @@ PRAGMA_ENABLE_DEPRECATION_WARNINGS
 		}
 
 		{	// launch a GT task, then an any-thread task that depends on it. wait for the any-thread task. this was a deadlock on the new frontend
-			FGraphEventRef GTTask = FFunctionGraphTask::CreateAndDispatchWhenReady([] { IsInGameThread(); }, TStatId{}, nullptr, ENamedThreads::GameThread);
+			FGraphEventRef GTTask = FFunctionGraphTask::CreateAndDispatchWhenReady([] { IsInGameThread(); }, TStatId{}, nullptr, ENamedThreads::GameThread); //-V530
 			FGraphEventRef AnyThreadTask = FFunctionGraphTask::CreateAndDispatchWhenReady([] {}, TStatId{}, GTTask);
 			AnyThreadTask->Wait();
 		}
diff --git a/Engine/Source/Runtime/Engine/Private/WorldPartition/WorldPartitionStreamingGeneration.cpp b/Engine/Source/Runtime/Engine/Private/WorldPartition/WorldPartitionStreamingGeneration.cpp
--- a/Engine/Source/Runtime/Engine/Private/WorldPartition/WorldPartitionStreamingGeneration.cpp
+++ b/Engine/Source/Runtime/Engine/Private/WorldPartition/WorldPartitionStreamingGeneration.cpp
@@ -873,7 +873,7 @@ class FWorldPartitionStreamingGenerator
 			else
 			{
 				const FGuid ActorGuid = InActorDescView.GetGuid();
-				OutActorDescViewMap.Emplace(ActorGuid, MoveTemp(InActorDescView));
+				OutActorDescViewMap.Emplace(ActorGuid, MoveTemp(InActorDescView)); //-V530
 			}
 		};
 
diff --git a/Engine/Source/Runtime/Experimental/Chaos/Public/Chaos/ISpatialAcceleration.h b/Engine/Source/Runtime/Experimental/Chaos/Public/Chaos/ISpatialAcceleration.h
--- a/Engine/Source/Runtime/Experimental/Chaos/Public/Chaos/ISpatialAcceleration.h
+++ b/Engine/Source/Runtime/Experimental/Chaos/Public/Chaos/ISpatialAcceleration.h
@@ -433,7 +433,7 @@ static TUniquePtr<TDerived> AsUniqueSpatialAcceleration(TUniquePtr<TBase>&& Base
 {
 	if (TDerived* Derived = Base->template As<TDerived>())
 	{
-		Base.Release();
+		(void)Base.Release();
 		return TUniquePtr<TDerived>(Derived);
 	}
 	return nullptr;
@@ -443,7 +443,7 @@ template <typename TDerived, typename TBase>
 static TUniquePtr<TDerived> AsUniqueSpatialAccelerationChecked(TUniquePtr<TBase>&& Base)
 {
 	TDerived& Derived = Base->template AsChecked<TDerived>();
-	Base.Release();
+	(void)Base.Release();
 	return TUniquePtr<TDerived>(&Derived);
 }
 
diff --git a/Engine/Source/Runtime/GeometryCore/Public/CircleTypes.h b/Engine/Source/Runtime/GeometryCore/Public/CircleTypes.h
--- a/Engine/Source/Runtime/GeometryCore/Public/CircleTypes.h
+++ b/Engine/Source/Runtime/GeometryCore/Public/CircleTypes.h
@@ -119,7 +119,6 @@ struct TCircle2
 
 	TAxisAlignedBox2<RealType> GetBoundingBox() const
 	{
-		TVector2<RealType>(Center.X + Radius, Center.Y + Radius);
 		return TAxisAlignedBox2<RealType>(
 			TVector2<RealType>(Center.X - Radius, Center.Y - Radius),
 			TVector2<RealType>(Center.X + Radius, Center.Y + Radius));
