[MassEntityTestSuite]
Properly deinitialize entity managers created for mass tests

diff --git a/Engine/Source/Developer/MassEntityTestSuite/Private/MassTestTypes.cpp b/Engine/Source/Developer/MassEntityTestSuite/Private/MassTestTypes.cpp
--- a/Engine/Source/Developer/MassEntityTestSuite/Private/MassTestTypes.cpp
+++ b/Engine/Source/Developer/MassEntityTestSuite/Private/MassTestTypes.cpp
@@ -24,6 +24,14 @@ bool FExecutionTestBase::SetUp()
 	return true;
 }
 
+void FExecutionTestBase::TearDown()
+{
+	FAITestBase::TearDown();
+
+	EntityManager->Deinitialize();
+	EntityManager = nullptr;
+}
+
 bool FEntityTestBase::SetUp()
 {
 	FExecutionTestBase::SetUp();
diff --git a/Engine/Source/Developer/MassEntityTestSuite/Public/MassEntityTestTypes.h b/Engine/Source/Developer/MassEntityTestSuite/Public/MassEntityTestTypes.h
--- a/Engine/Source/Developer/MassEntityTestSuite/Public/MassEntityTestTypes.h
+++ b/Engine/Source/Developer/MassEntityTestSuite/Public/MassEntityTestTypes.h
@@ -381,6 +381,7 @@ struct FExecutionTestBase : FAITestBase
 	bool bMakeWorldEntityManagersOwner = true;
 
 	UE_API virtual bool SetUp() override;
+	UE_API virtual void TearDown() override;
 };
 
 struct FEntityTestBase : FExecutionTestBase
