[AINavigation] fix mass entity initialization issue caused by AutoRTFM transaction side effect on scoped FCreationContext that are never destroyed when leaving a scope.
#rnx

diff --git a/Engine/Source/Developer/MassEntityTestSuite/Private/MassCreationContextTest.cpp b/Engine/Source/Developer/MassEntityTestSuite/Private/MassCreationContextTest.cpp
--- a/Engine/Source/Developer/MassEntityTestSuite/Private/MassCreationContextTest.cpp
+++ b/Engine/Source/Developer/MassEntityTestSuite/Private/MassCreationContextTest.cpp
@@ -121,6 +121,54 @@ struct FManualBuild : FEntityTestBase
 };
 IMPLEMENT_AI_INSTANT_TEST(FManualBuild, "System.Mass.CreationContext.ManualBuild");
 
+// Requires a fix for UE-333172
+//struct FCreationContextScopeTest : FEntityTestBase
+//{
+//	virtual bool InstantTest() override
+//	{
+//		TArray<FMassEntityHandle> BothEntities;
+//		constexpr int FirstEntityValue = 1;
+//		constexpr int SecondEntityValue = 2;
+//
+//		{
+//			TArray<FMassEntityHandle> Entities;
+//			const TSharedRef<FMassEntityManager::FEntityCreationContext> CreationContext = EntityManager->BatchCreateEntities(IntsArchetype, 1, Entities);
+//			BothEntities.Append(Entities);
+//
+//			TArray<FMassArchetypeEntityCollection> EntityCollections = CreationContext->GetEntityCollections(*EntityManager.Get());
+//			AITEST_EQUAL(TEXT("We expect the number of resulting collections to match expectations"), EntityCollections.Num(), 1);
+//			AITEST_EQUAL(TEXT("We expect the number of entities in resulting collection to match expectations"), EntityCollections[0].DebugCountEntities(), 1);
+//
+//			EntityManager->BatchSetEntityFragmentValues(CreationContext->GetEntityCollections(*EntityManager), {FInstancedStruct::Make(FTestFragment_Int{FirstEntityValue})});
+//		}
+//
+//		{
+//			TArray<FMassEntityHandle> Entities;
+//			const TSharedRef<FMassEntityManager::FEntityCreationContext> CreationContext = EntityManager->BatchCreateEntities(IntsArchetype, 1, Entities);
+//			BothEntities.Append(Entities);
+//
+//			TArray<FMassArchetypeEntityCollection> EntityCollections = CreationContext->GetEntityCollections(*EntityManager.Get());
+//			AITEST_EQUAL(TEXT("We expect the number of resulting collections to match expectations"), EntityCollections.Num(), 1);
+//			AITEST_EQUAL(TEXT("We expect the number of entities in resulting collection to match expectations"), EntityCollections[0].DebugCountEntities(), 1);
+//
+//			EntityManager->BatchSetEntityFragmentValues(CreationContext->GetEntityCollections(*EntityManager), {FInstancedStruct::Make(FTestFragment_Int{SecondEntityValue})});
+//		}
+//
+//		AITEST_EQUAL(TEXT("We expect the total number of entities to match expectations"), BothEntities.Num(), 2);
+//
+//		FTestFragment_Int* FirstEntityFragment = EntityManager->GetFragmentDataPtr<FTestFragment_Int>(BothEntities[0]);
+//		AITEST_NOT_NULL(TEXT("The first entity is expected to have the int fragment"), FirstEntityFragment);
+//		AITEST_EQUAL(TEXT("We expect the first entity's fragment to match expectations"), FirstEntityFragment->Value, FirstEntityValue);
+//
+//		FTestFragment_Int* SecondEntityFragment = EntityManager->GetFragmentDataPtr<FTestFragment_Int>(BothEntities[1]);
+//		AITEST_NOT_NULL(TEXT("The second entity is expected to have the int fragment"), SecondEntityFragment);
+//		AITEST_EQUAL(TEXT("We expect the second entity's fragment to match expectations"), SecondEntityFragment->Value, SecondEntityValue);
+//
+//		return true;
+//	}
+//};
+//IMPLEMENT_AI_INSTANT_TEST_WITH_FLAGS(FCreationContextScopeTest, "System.Mass.CreationContext.ScopedOperations", EAutomationTestFlags::ClientContext | EAutomationTestFlags::EditorContext | EAutomationTestFlags::EngineFilter | EAutomationTestFlags::SupportsAutoRTFM)
+
 } // UE::Mass::Test::CreationContext
 
 UE_ENABLE_OPTIMIZATION_SHIP
