[Mass]
Fix linked entity type checking

#code-analyzed

diff --git a/Engine/Source/Developer/MassEntityTestSuite/Private/MassQueryTest.cpp b/Engine/Source/Developer/MassEntityTestSuite/Private/MassQueryTest.cpp
--- a/Engine/Source/Developer/MassEntityTestSuite/Private/MassQueryTest.cpp
+++ b/Engine/Source/Developer/MassEntityTestSuite/Private/MassQueryTest.cpp
@@ -1086,7 +1086,7 @@ struct FQueryTest_LinkedFragmentAccess : FEntityTestBase
 		FMassExecutionContext ExecContext(*EntityManager.Get());
 		Query.ForEachEntityChunk(ExecContext, [this, &bAnyMatch, &bFoundShared](FMassExecutionContext& Context)
 			{
-				const FTestFragment_Int* SharedInt = Context.GetLinkedEntityFragmentPtr<FTestFragment_Int>();
+				const FTestFragment_Int* SharedInt = Context.GetLinkedEntityConstFragmentPtr<FTestFragment_Int>();
 				if (!SharedInt)
 				{
 					bFoundShared = false;
@@ -1138,7 +1138,7 @@ struct FQueryTest_IndirectFragmentAccess : FEntityTestBase
 
 		FMassEntityQuery Query(EntityManager);
 		Query.AddRequirement<FTestFragment_EntityHandle>(EMassFragmentAccess::ReadOnly);
-		Query.AddIndirectFragmentRequirement<FTestFragment_Float>(EMassFragmentAccess::ReadOnly);
+		Query.AddIndirectFragmentRequirement<FTestFragment_Int>(EMassFragmentAccess::ReadOnly);
 		Query.CacheArchetypes();
 
 		bool bAllMatch = true;
diff --git a/Engine/Source/Runtime/MassEntity/Public/MassRequirements.h b/Engine/Source/Runtime/MassEntity/Public/MassRequirements.h
--- a/Engine/Source/Runtime/MassEntity/Public/MassRequirements.h
+++ b/Engine/Source/Runtime/MassEntity/Public/MassRequirements.h
@@ -527,7 +527,7 @@ struct FMassFragmentRequirements
 				return Item.StructType == LinkType;
 			});
 
-			if (bHasLinkType)
+		if (!bHasLinkType)
 		{
 			checkfSlow(ConstSharedFragmentRequirements.ContainsByPredicate([](const FMassFragmentRequirementDescription& Item)
 				{
