[Mass] Added a flag to EntityBuilder to (optionally) force committing always going through Mass commands. There's also a cvar that can be used to set the behavior project-wide (mass.EntityBuilder.ForceDeferredCommit).

#Synced-CL 49615561

diff --git a/Engine/Source/Developer/MassEntityTestSuite/Private/MassEntityBuilderTest.cpp b/Engine/Source/Developer/MassEntityTestSuite/Private/MassEntityBuilderTest.cpp
--- a/Engine/Source/Developer/MassEntityTestSuite/Private/MassEntityBuilderTest.cpp
+++ b/Engine/Source/Developer/MassEntityTestSuite/Private/MassEntityBuilderTest.cpp
@@ -33,7 +33,8 @@ namespace UE::Mass::Test::EntityBuilder
 
 			{
 				FEntityBuilder EntityBuilder(*EntityManager.Get());
-				EntityBuilder.Add_GetRef<FTestFragment_Int>(SomeCounter++);
+				EntityBuilder.SetForceDeferredCommit(false)
+					.Add_GetRef<FTestFragment_Int>(SomeCounter++);
 				EntityBuilder.Commit();
 				EntitiesCreatedThisStep = 1;
 			}
@@ -43,6 +44,8 @@ namespace UE::Mass::Test::EntityBuilder
 #endif // WITH_MASSENTITY_DEBUG
 			{
 				FEntityBuilder EntityBuilder(*EntityManager.Get());
+				EntityBuilder.SetForceDeferredCommit(false);
+
 				ON_SCOPE_EXIT
 				{
 					EntityBuilder.Commit();
@@ -57,7 +60,8 @@ namespace UE::Mass::Test::EntityBuilder
 #endif // WITH_MASSENTITY_DEBUG
 			{
 				FScopedEntityBuilder EntityBuilder(*EntityManager.Get());
-				EntityBuilder.Add_GetRef<FTestFragment_Int>(SomeCounter++);
+				EntityBuilder.SetForceDeferredCommit(false)
+					.Add_GetRef<FTestFragment_Int>(SomeCounter++);
 			}
 			EntitiesCreatedThisStep = 1;
 #if WITH_MASSENTITY_DEBUG
@@ -66,7 +70,8 @@ namespace UE::Mass::Test::EntityBuilder
 #endif // WITH_MASSENTITY_DEBUG
 			{
 				FScopedEntityBuilder EntityBuilder(*EntityManager.Get());
-				EntityBuilder.Add_GetRef<FTestFragment_Int>(SomeCounter++);
+				EntityBuilder.SetForceDeferredCommit(false)
+					.Add_GetRef<FTestFragment_Int>(SomeCounter++);
 
 				FEntityBuilder EntityBuilder2 = EntityBuilder;
 				EntityBuilder2.Commit();
@@ -78,7 +83,8 @@ namespace UE::Mass::Test::EntityBuilder
 #endif // WITH_MASSENTITY_DEBUG
 			{
 				FEntityBuilder EntityBuilder(*EntityManager.Get());
-				EntityBuilder.Add_GetRef<FTestFragment_Int>(SomeCounter++);
+				EntityBuilder.SetForceDeferredCommit(false)
+					.Add_GetRef<FTestFragment_Int>(SomeCounter++);
 				EntityBuilder.Commit();
 
 				FEntityBuilder EntityBuilder2 = EntityBuilder;
@@ -90,7 +96,8 @@ namespace UE::Mass::Test::EntityBuilder
 #endif // WITH_MASSENTITY_DEBUG
 			{
 				FEntityBuilder EntityBuilder(*EntityManager.Get());
-				EntityBuilder.Add_GetRef<FTestFragment_Int>(SomeCounter++);
+				EntityBuilder.SetForceDeferredCommit(false)
+					.Add_GetRef<FTestFragment_Int>(SomeCounter++);
 
 				FEntityBuilder EntityBuilder2 = EntityBuilder;
 				EntityBuilder2.Commit();
@@ -149,7 +156,9 @@ namespace UE::Mass::Test::EntityBuilder
 		{
 			int32 TotalCountCreated = 0;
 			{
-				FMassEntityHandle CreatedEntity = EntityManager->MakeEntityBuilder().Add<FTestFragment_Int>().Commit();
+				FMassEntityHandle CreatedEntity = EntityManager->MakeEntityBuilder()
+					.SetForceDeferredCommit(false)
+					.Add<FTestFragment_Int>().Commit();
 				++TotalCountCreated;
 
 				const bool bIsBuilt = EntityManager->IsEntityActive(CreatedEntity);
@@ -188,10 +197,12 @@ namespace UE::Mass::Test::EntityBuilder
 		{
 			const int32 ValueA = 1;
 			FEntityBuilder BuilderA(*EntityManager.Get());
+			BuilderA.SetForceDeferredCommit(false);
 			BuilderA.Add<FTestFragment_Int>(ValueA);
 
 			const int32 ValueB = ValueA + 1;
 			FEntityBuilder BuilderB = EntityManager->MakeEntityBuilder();
+			BuilderB.SetForceDeferredCommit(false);
 			BuilderB.Add<FTestFragment_Int>(ValueB);
 
 			// a different way of setting the value
@@ -280,6 +291,7 @@ namespace UE::Mass::Test::EntityBuilder
 
 			{
 				FEntityBuilder BuilderA(*EntityManager.Get());
+				BuilderA.SetForceDeferredCommit(false);
 				BuilderA.Add<FTestFragment_Int>();
 				// forces handle reservation
 				FMassEntityHandle EntityA = BuilderA.Commit();
@@ -292,7 +304,7 @@ namespace UE::Mass::Test::EntityBuilder
 					BuilderA = MoveTemp(BuilderB);
 				}
 
-				// at this point EntityB should be valid while the original EntityA not
+				// at this point both EntityA and EntityB should be valid
 				AITEST_TRUE("The original entity is valid, since it was committed", EntityManager->IsEntityValid(EntityA));
 				AITEST_TRUE("The original entity is active", EntityManager->IsEntityActive(EntityA));
 				AITEST_TRUE("The passed-over entity entity is valid", EntityManager->IsEntityValid(EntityB));
@@ -307,12 +319,13 @@ namespace UE::Mass::Test::EntityBuilder
 				FMassEntityHandle EntityB;
 				{
 					FEntityBuilder BuilderB(*EntityManager.Get());
+					BuilderB.SetForceDeferredCommit(false);
 					BuilderB.Add<FTestFragment_Float>();
 					EntityB = BuilderB.Commit();
 					BuilderA = MoveTemp(BuilderB);
 				}
 
-				// at this point EntityB should be valid while the original EntityA not
+				// at this point both EntityA and EntityB should be valid, but only EntityB should be active (commited).
 				AITEST_TRUE("The original entity is valid", EntityManager->IsEntityValid(EntityA));
 				AITEST_TRUE("The original entity is NOT active", EntityManager->IsEntityActive(EntityA) == false);
 				AITEST_TRUE("The secondary entity entity is valid", EntityManager->IsEntityValid(EntityB));
@@ -321,19 +334,21 @@ namespace UE::Mass::Test::EntityBuilder
 
 			{
 				FEntityBuilder BuilderA(*EntityManager.Get());
-				BuilderA.Add<FTestFragment_Int>();
+				BuilderA.SetForceDeferredCommit(false)
+					.Add<FTestFragment_Int>();
 				// forces handle reservation
 				FMassEntityHandle EntityA = BuilderA.Commit();
 
 				FMassEntityHandle EntityB;
 				{
 					FEntityBuilder BuilderB(*EntityManager.Get());
-					BuilderB.Add<FTestFragment_Float>();
+					BuilderB.SetForceDeferredCommit(false)
+						.Add<FTestFragment_Float>();
 					EntityB = BuilderB.Commit();
 					BuilderA = MoveTemp(BuilderB);
 				}
 
-				// at this point EntityB should be valid while the original EntityA not
+				// at this point both EntityA and EntityB should be valid and active (commited). 
 				AITEST_TRUE("The original entity is valid", EntityManager->IsEntityValid(EntityA));
 				AITEST_TRUE("The original entity is active", EntityManager->IsEntityActive(EntityA));
 				AITEST_TRUE("The secondary entity entity is valid", EntityManager->IsEntityValid(EntityB));
@@ -357,7 +372,8 @@ namespace UE::Mass::Test::EntityBuilder
 		virtual bool InstantTest() override
 		{
 			FEntityBuilder Builder = EntityManager->MakeEntityBuilder();
-			Builder.Add<FTestFragment_Int>();
+			Builder.SetForceDeferredCommit(false)
+				.Add<FTestFragment_Int>();
 
 			auto TestBuilder = [this, &Builder](const FMassArchetypeHandle& ExpectedArchetype)-> bool
 			{
@@ -446,7 +462,8 @@ namespace UE::Mass::Test::EntityBuilder
 			EntityManager->CreateEntity(IntsArchetype);
 
 			FEntityBuilder SyncBuilder = EntityManager->MakeEntityBuilder();
-			SyncBuilder.Add<FTestFragment_Int>();
+			SyncBuilder.SetForceDeferredCommit(false)
+				.Add<FTestFragment_Int>();
 
 			const FMassEntityHandle ReservedHandle = SyncBuilder.GetEntityHandle();
 
@@ -504,6 +521,7 @@ namespace UE::Mass::Test::EntityBuilder
 			{
 				// quick builder just to create an entity with known properties
 				FEntityBuilder Builder(EntityManager.ToSharedRef());
+				Builder.SetForceDeferredCommit(false);
 				Builder.Add<FTestFragment_Int>(TestIntValue);
 				Builder.Add<FTestFragment_Float>(TestFloatValue);
 				Builder.Add<FTestTag_B>();
@@ -566,6 +584,7 @@ namespace UE::Mass::Test::EntityBuilder
 		virtual bool InstantTest() override
 		{
 			FEntityBuilder Builder(EntityManager.ToSharedRef());
+			Builder.SetForceDeferredCommit(false);
 			Builder.CopyDataFromEntity(OriginalEntity);
 
 			const FMassEntityHandle NewEntityHandle = Builder.Commit();
@@ -587,7 +606,7 @@ namespace UE::Mass::Test::EntityBuilder
 			Builder.Add<FTestTag_A>();
 			Builder.AppendDataFromEntity(OriginalEntity);
 
-			const FMassEntityHandle NewEntityHandle = Builder.Commit();
+			const FMassEntityHandle NewEntityHandle = Builder.SetForceDeferredCommit(false).Commit();
 
 			EntityManager->RemoveTagFromEntity(NewEntityHandle, FTestTag_A::StaticStruct());
 
@@ -608,6 +627,7 @@ namespace UE::Mass::Test::EntityBuilder
 			{
 				// quick builder just to create an entity with known properties
 				FEntityBuilder Builder(EntityManager.ToSharedRef());
+				Builder.SetForceDeferredCommit(false);
 
 				FInstancedStruct ElementInstance;
 
diff --git a/Engine/Source/Developer/MassEntityTestSuite/Private/MassEntityTest.cpp b/Engine/Source/Developer/MassEntityTestSuite/Private/MassEntityTest.cpp
--- a/Engine/Source/Developer/MassEntityTestSuite/Private/MassEntityTest.cpp
+++ b/Engine/Source/Developer/MassEntityTestSuite/Private/MassEntityTest.cpp
@@ -443,7 +443,7 @@ struct FEntityTest_CreationPatterns : FEntityTestBase
 		// done automatically, but subsequent explicit Build calls will fail
 
 		// entity builder, new in 5.6. There's more, see MassEntityBuilderTest.cpp
-		Entities.Add(EntityManager->MakeEntityBuilder()
+		Entities.Add(EntityManager->MakeEntityBuilder().SetForceDeferredCommit(false)
 			.Add<FTestFragment_Int>(1024)
 			.Add<FTestFragment_Float>(3.14f)
 			.Commit());
diff --git a/Engine/Source/Developer/MassEntityTestSuite/Private/MassRelationTest.cpp b/Engine/Source/Developer/MassEntityTestSuite/Private/MassRelationTest.cpp
--- a/Engine/Source/Developer/MassEntityTestSuite/Private/MassRelationTest.cpp
+++ b/Engine/Source/Developer/MassEntityTestSuite/Private/MassRelationTest.cpp
@@ -1,6 +1,7 @@
 // Copyright Epic Games, Inc. All Rights Reserved.
 
 #include "AITestsCommon.h"
+#include "MassEntityBuilder.h"
 #include "MassRelationManager.h"
 #include "MassEntityTestTypes.h"
 
@@ -237,6 +238,13 @@ namespace UE::Mass::Test::Relations
 	{
 		virtual bool InstantTest() override
 		{
+			if (FEntityBuilder::GetDefaultValueOfForceDeferredCommit())
+			{
+				// this test assumes synchronous EntityBuilder commits. Bailing out
+				// to avoid generating false positives
+				return true;
+			}
+
 			constexpr int32 NumEntities = 3;
 			const FMassArchetypeHandle OriginalArchetype = IntsArchetype;
 
@@ -574,6 +582,13 @@ namespace UE::Mass::Test::Relations
 	{
 		virtual bool InstantTest() override
 		{
+			if (FEntityBuilder::GetDefaultValueOfForceDeferredCommit())
+			{
+				// this test assumes synchronous EntityBuilder commits. Bailing out
+				// to avoid generating false positives
+				return true;
+			}
+
 			constexpr int32 NumEntities = 3;
 			constexpr int32 OriginalParentIndex = 0;
 			constexpr int32 ChildIndex = 1;
@@ -608,6 +623,13 @@ namespace UE::Mass::Test::Relations
 	{
 		virtual bool InstantTest() override
 		{
+			if (FEntityBuilder::GetDefaultValueOfForceDeferredCommit())
+			{
+				// this test assumes synchronous EntityBuilder commits. Bailing out
+				// to avoid generating false positives
+				return true;
+			}
+
 			constexpr int32 NumEntities = 3;
 			const FMassArchetypeHandle OriginalArchetype = IntsArchetype;
 
diff --git a/Engine/Source/Runtime/MassEntity/Private/MassEntityBuilder.cpp b/Engine/Source/Runtime/MassEntity/Private/MassEntityBuilder.cpp
--- a/Engine/Source/Runtime/MassEntity/Private/MassEntityBuilder.cpp
+++ b/Engine/Source/Runtime/MassEntity/Private/MassEntityBuilder.cpp
@@ -26,6 +26,18 @@ namespace UE::Mass
 	}
 #endif // WITH_MASSENTITY_DEBUG
 
+	namespace Tweakable
+	{
+		bool bForceDeferredEntityBuilderCommit = false;
+		namespace
+		{
+			FAutoConsoleVariableRef AnonymousCVars[] = {
+				{TEXT("mass.EntityBuilder.ForceDeferredCommit"), bForceDeferredEntityBuilderCommit
+					, TEXT("When set will result in each EntityBuilder commit its work using commands, regardless of EntityManager.IsProcessing value.")}
+			};
+		}
+	}
+
 	namespace Private
 	{
 		struct FEntityBuilderHelper
@@ -120,11 +132,13 @@ namespace UE::Mass
 //-----------------------------------------------------------------------------
 FEntityBuilder::FEntityBuilder(FMassEntityManager& InEntityManager)
 	: EntityManager(InEntityManager.AsShared())
+	, bForceDeferred(UE::Mass::Tweakable::bForceDeferredEntityBuilderCommit)
 {	
 }
 
 FEntityBuilder::FEntityBuilder(const TSharedRef<FMassEntityManager>& InEntityManager)
 	: EntityManager(InEntityManager)
+	, bForceDeferred(UE::Mass::Tweakable::bForceDeferredEntityBuilderCommit)
 {
 	
 }
@@ -160,6 +174,7 @@ FEntityBuilder& FEntityBuilder::operator=(const FEntityBuilder& Other)
 		ConstSharedFragments = Other.ConstSharedFragments;
 
 		State = Composition.IsEmpty() ? EState::Empty : EState::ReadyToCommit;
+		bForceDeferred = Other.bForceDeferred;
 	}
 
 	return *this;
@@ -199,6 +214,7 @@ FEntityBuilder& FEntityBuilder::operator=(FEntityBuilder&& Other)
 			EntityHandle = Other.EntityHandle;
 			State = Other.State;
 		}
+		bForceDeferred = Other.bForceDeferred;
 
 		Other.EntityHandle.Reset();
 		Other.State = EState::Invalid;
@@ -276,7 +292,7 @@ FMassEntityHandle FEntityBuilder::Commit()
 	CacheSharedFragmentValue();
 	CacheArchetypeHandle();
 
-	if (EntityManager->IsProcessing())
+	if (bForceDeferred || EntityManager->IsProcessing())
 	{
 		// we need to issue commands in this case
 		EntityManager->Defer().PushCommand<FMassDeferredCreateCommand>(
@@ -597,4 +613,9 @@ void FEntityBuilder::ForEachRelation(const TFunctionRef<bool(FPendingRelationPar
 	}
 }
 
+bool FEntityBuilder::GetDefaultValueOfForceDeferredCommit()
+{
+	return Tweakable::bForceDeferredEntityBuilderCommit;
+}
+
 } // namespace UE::Mass
\ No newline at end of file
diff --git a/Engine/Source/Runtime/MassEntity/Public/MassEntityBuilder.h b/Engine/Source/Runtime/MassEntity/Public/MassEntityBuilder.h
--- a/Engine/Source/Runtime/MassEntity/Public/MassEntityBuilder.h
+++ b/Engine/Source/Runtime/MassEntity/Public/MassEntityBuilder.h
@@ -91,6 +91,12 @@ struct FEntityBuilder
 		, TArray<FConstSharedStruct>&& ConstSharedFragments
 		, TArray<FSharedStruct>&& SharedFragments);
 
+	/**
+	 * Configured given EntityBuilder instance's commit behavior. Using `true` will result in all Commit calls
+	 * resulting in issuing Mass commands, regardless of EntityManager.IsProcessing() value.
+	 */
+	FEntityBuilder& SetForceDeferredCommit(bool bNewValue);
+
 	/**
 	 * Finalizes the creation of the entity with the specified fragments and configurations.
 	 * Note that this function needs to be called manually, no automated entity creation will take place upon builder's destruction. 
@@ -242,6 +248,9 @@ struct FEntityBuilder
 	/** @return the EntityManager instance this entity builder is working for */
 	TSharedRef<FMassEntityManager> GetEntityManager();
 
+	/** Retrieves the default value used to initialize EntityBuilder.bForceDeferred */
+	UE_API static bool GetDefaultValueOfForceDeferredCommit();
+
 protected:
 	friend Private::FEntityBuilderHelper;
 
@@ -321,6 +330,9 @@ private:
 		Invalid,
 	};
 	EState State = EState::Empty;
+
+	/** if set will make entity creation always performed with Mass commands */
+	uint8 bForceDeferred : 1 = false;
 };
 
 struct FScopedEntityBuilder : FEntityBuilder
@@ -454,6 +466,12 @@ inline void FEntityBuilder::ConfigureArchetypeCreation(const FMassArchetypeCreat
 	ArchetypeCreationParams = InCreationParams;
 }
 
+inline FEntityBuilder& FEntityBuilder::SetForceDeferredCommit(const bool bNewValue)
+{
+	bForceDeferred = bNewValue;
+	return *this;
+}
+
 } // namespace UE::Mass
 
 #undef UE_API
\ No newline at end of file
