[Mass]
Mass QueryExecutor subsystem accessors now use templated subsystem getter from the execution context.
Fixed implicit casting ambiguity with Mass QueryExecutor subsystem accessors

#Synced-CL 49613770

diff --git a/Engine/Source/Developer/MassEntityTestSuite/Private/MassQueryExecutorTest.cpp b/Engine/Source/Developer/MassEntityTestSuite/Private/MassQueryExecutorTest.cpp
--- a/Engine/Source/Developer/MassEntityTestSuite/Private/MassQueryExecutorTest.cpp
+++ b/Engine/Source/Developer/MassEntityTestSuite/Private/MassQueryExecutorTest.cpp
@@ -156,8 +156,9 @@ namespace FMassQueryExecutorTest
 
 		virtual void Execute(FMassExecutionContext& Context)
 		{
-			ParallelForEachEntityChunk(Context, Accessors, [](FMassExecutionContext& Context, auto& Data)
+			ForEachEntityChunk(Context, Accessors, [](FMassExecutionContext& Context, auto& Data)
 			{
+				UMassTestWorldSubsystem& TestSubSys = Data.template Get<UMassTestWorldSubsystem>();
 				for (uint32 EntityIndex : Context.CreateEntityIterator())
 				{
 					Data.template Get<FTestFragment_Int>()[EntityIndex].Value = 0;
@@ -165,7 +166,7 @@ namespace FMassQueryExecutorTest
 					Data.template Get<FTestFragment_Bool>()[EntityIndex].bValue = FMath::RandBool();
 					Data.template Get<FTestFragment_Large>()[EntityIndex].Value[0] = FMath::RandBool() ? 1 : 0;
 					Data.template Get<FTestFragment_Array>()[EntityIndex].Value.SetNum(FMath::RandBool() ? 1 : 0);
-					Data.template Get<UMassTestWorldSubsystem>()->Write(0);
+					TestSubSys.Write(0);
 				}
 			});
 		}
@@ -883,6 +884,9 @@ namespace FMassQueryExecutorTest
 			{
 				UE_LOG(LogMassPerfTest, Log, TEXT("\n%d entities:"), EntityCount);
 				CreateEntities(EntityCount);
+				Processors.Reset();
+				Processors.Add(Processor_Setup);
+				UE::Mass::Executor::RunProcessorsView(Processors, ProcessingContext);
 				TArray<TestReturn> Results;
 				int32 DefaultIndex = Results.Emplace(RunTest(DefaultProcessor, nullptr, EntityCount, RunCount, TEXT("Default")));
 				TestReturn Default = Results[DefaultIndex];
@@ -916,8 +920,7 @@ namespace FMassQueryExecutorTest
 			return true;
 		}
 	};
-	// This test can't fail and only exists to compare relative performance of various processor implementations and environments
-	//IMPLEMENT_AI_INSTANT_TEST(FQueryExecutor_LoadTest, "System.Mass.Processor.AutoExecuteQuery.LoadTest");
+	IMPLEMENT_AI_INSTANT_TEST(FQueryExecutor_LoadTest, "System.Mass.Processor.AutoExecuteQuery.LoadTest");
 
 } // FMassQueryExecutorTest
 
diff --git a/Engine/Source/Runtime/MassEntity/Public/MassQueryExecutor.h b/Engine/Source/Runtime/MassEntity/Public/MassQueryExecutor.h
--- a/Engine/Source/Runtime/MassEntity/Public/MassQueryExecutor.h
+++ b/Engine/Source/Runtime/MassEntity/Public/MassQueryExecutor.h
@@ -209,10 +209,10 @@ struct FQueryDefinition : public FQueryDefinitionBase
 	}
 
 	template <typename TFragment>
-	inline constexpr auto& Get()
+	constexpr auto& Get()
 	{
 		constexpr std::size_t Index = GetAccessorIndex<TFragment, Accessors...>();
-		return AccessorTuple.template Get<Index>();
+		return AccessorTuple.template Get<Index>().Get();
 	}
 };
 
@@ -708,9 +708,10 @@ struct FMutableSubsystemAccess
 
 	FMutableSubsystemAccess() = default;
 
-	inline TSubsystem* Get() const
+	TSubsystem& Get() const
 	{
-		return Subsystem;
+		check(Subsystem);
+		return *Subsystem;
 	}
 
 	inline TSubsystem& operator*() const
@@ -756,8 +757,9 @@ struct FConstSubsystemAccess
 
 	FConstSubsystemAccess() = default;
 
-	inline const TSubsystem* Get() const
+	const TSubsystem& Get() const
 	{
+		check(Subsystem);
 		return Subsystem;
 	}
 
@@ -784,7 +786,7 @@ struct FConstSubsystemAccess
 
 	inline void SetupForExecute(FMassExecutionContext& Context)
 	{
-		Subsystem = Context.GetSubsystem(TSubsystem::StaticClass());
+		Subsystem = Context.GetSubsystem<TSubsystem>();
 	}
 
 	inline void SetupForChunk(FMassExecutionContext& Context)
