[Mass] Fixed System.Mass.EntityBuilder.OneLiner unit test to not fail if EntityBuilder is configured globally to always rely on commands

#rnx
#Synced-CL 49674283

diff --git a/Engine/Source/Developer/MassEntityTestSuite/Private/MassEntityBuilderTest.cpp b/Engine/Source/Developer/MassEntityTestSuite/Private/MassEntityBuilderTest.cpp
--- a/Engine/Source/Developer/MassEntityTestSuite/Private/MassEntityBuilderTest.cpp
+++ b/Engine/Source/Developer/MassEntityTestSuite/Private/MassEntityBuilderTest.cpp
@@ -168,7 +168,9 @@ namespace UE::Mass::Test::EntityBuilder
 #endif // WITH_MASSENTITY_DEBUG
 			}
 			{
-				FEntityBuilder EntityBuilder = EntityManager->MakeEntityBuilder().Add<FTestFragment_Int>();
+				FEntityBuilder EntityBuilder = EntityManager->MakeEntityBuilder()
+					.SetForceDeferredCommit(false)
+					.Add<FTestFragment_Int>();
 				EntityBuilder.Commit();
 				++TotalCountCreated;
 			}
