[Mass]
Entity creation in Processors, including off main thread.
* Require processors & queries to declare which archetypes they will create
* We create the archetypes in advance
* We schedule processors such that they won't contend for access to the listed archetypes
* There is no synchronization at entity creation time

Added automated integration tests for processor-scheduled entity creation, including scheduling and parallel execution

#Synced-CL 50088124

diff --git a/Engine/Source/Developer/MassEntityTestSuite/Private/MassComplexTest.cpp b/Engine/Source/Developer/MassEntityTestSuite/Private/MassComplexTest.cpp
--- a/Engine/Source/Developer/MassEntityTestSuite/Private/MassComplexTest.cpp
+++ b/Engine/Source/Developer/MassEntityTestSuite/Private/MassComplexTest.cpp
@@ -4,11 +4,15 @@
 #include "MassEntityManager.h"
 #include "MassEntityTestTypes.h"
 #include "MassExecutionContext.h"
+#include "MassProcessingPhaseManager.h"
+#include "MassCommands.h"
 
 #define LOCTEXT_NAMESPACE "MassTest"
 
 UE_DISABLE_OPTIMIZATION_SHIP
 
+DEFINE_LOG_CATEGORY_STATIC(LogMassComplexTest, Log, All);
+
 //-----------------------------------------------------------------------------
 // This file contains tests spanning multiple functionalities.
 //-----------------------------------------------------------------------------
@@ -83,6 +87,710 @@ struct FComplex_AddingFragmentInstancesToDiscontinouousEntitiesCollection : FEnt
 };
 IMPLEMENT_AI_INSTANT_TEST(FComplex_AddingFragmentInstancesToDiscontinouousEntitiesCollection, "System.Mass.Complex.AddingFragmentInstancesToDiscontinouousEntitiesCollection");
 
+//-----------------------------------------------------------------------------
+// Entity creation in scheduled processors
+//-----------------------------------------------------------------------------
+struct FComplex_CreateEntitiesParallel : FProcessingPhasesTestBase
+{
+	using Super = FProcessingPhasesTestBase;
+
+	static constexpr int32 NumToSpawn = 1000000;
+	static constexpr EMassProcessingPhase PhaseToRun = EMassProcessingPhase::PrePhysics;
+
+	TObjectPtr<UMassTestProcessorBase> CreateIntsProcessor = nullptr;
+	TObjectPtr<UMassTestProcessorBase> CreateFloatsProcessor = nullptr;
+	TObjectPtr<UMassTestProcessorBase> CreateLargeProcessor = nullptr;
+
+	FMassElementBitSet IntsComposition;
+	FMassElementBitSet FloatsComposition;
+	FMassElementBitSet LargeComposition;
+
+	virtual bool PopulatePhasesConfig() override
+	{
+		return true;
+	}
+
+	virtual bool SetUp() override
+	{
+		if (!Super::SetUp())
+		{
+			return false;
+		}
+
+		IntsComposition = IntsArchetype.GetCompositionBitSetChecked();
+		CreateIntsProcessor = NewObject<UMassTestProcessorBase>(GetTransientPackage());
+		CreateIntsProcessor->SetShouldAllowMultipleInstances(true);
+		CreateIntsProcessor->SetProcessingPhase(PhaseToRun);
+		CreateIntsProcessor->GetEntityCreationRequirements().AddCreatedArchetype(IntsComposition);
+		CreateIntsProcessor->ExecutionFunction = [this](FMassEntityManager& ExecuteEntityManager, FMassExecutionContext& Context)
+			{
+				UE_LOG(LogMassComplexTest, Log, TEXT("Begin Creating IntsArchetype"));
+
+				TArray<FMassEntityHandle> Entities;
+				EntityManager->BatchReserveEntities(NumToSpawn, Entities);
+				for (int32 i = 0; i < NumToSpawn; ++i)
+				{
+					Context.BuildEntity(IntsArchetype, Entities[i]);
+					FMassEntityView EntityView(ExecuteEntityManager, Entities[i]);
+					EntityView.GetFragmentData<FTestFragment_Int>().Value = 1;
+				}
+
+				UE_LOG(LogMassComplexTest, Log, TEXT("End Creating IntsArchetype"));
+			};
+		PhaseManager->RegisterDynamicProcessor(*CreateIntsProcessor);
+
+		FloatsComposition = FloatsArchetype.GetCompositionBitSetChecked();
+		CreateFloatsProcessor = NewObject<UMassTestProcessorBase>(GetTransientPackage());
+		CreateFloatsProcessor->SetShouldAllowMultipleInstances(true);
+		CreateFloatsProcessor->SetProcessingPhase(PhaseToRun);
+		CreateFloatsProcessor->GetEntityCreationRequirements().AddCreatedArchetype(FloatsComposition);
+		CreateFloatsProcessor->ExecutionFunction = [this](FMassEntityManager& ExecuteEntityManager, FMassExecutionContext& Context)
+			{
+				UE_LOG(LogMassComplexTest, Log, TEXT("Begin Creating FloatsArchetype"));
+
+				TArray<FMassEntityHandle> Entities;
+				EntityManager->BatchReserveEntities(NumToSpawn, Entities);
+				for (int32 i = 0; i < NumToSpawn; ++i)
+				{
+					Context.BuildEntity(FloatsArchetype, Entities[i]);
+					FMassEntityView EntityView(ExecuteEntityManager, Entities[i]);
+					EntityView.GetFragmentData<FTestFragment_Float>().Value = 1.0f;
+				}
+
+				UE_LOG(LogMassComplexTest, Log, TEXT("End Creating FloatsArchetype"));
+			};
+		PhaseManager->RegisterDynamicProcessor(*CreateFloatsProcessor);
+
+		LargeComposition = LargeArchetype.GetCompositionBitSetChecked();
+		CreateLargeProcessor = NewObject<UMassTestProcessorBase>(GetTransientPackage());
+		CreateLargeProcessor->SetShouldAllowMultipleInstances(true);
+		CreateLargeProcessor->SetProcessingPhase(PhaseToRun);
+		CreateLargeProcessor->GetEntityCreationRequirements().AddCreatedArchetype(LargeComposition);
+		CreateLargeProcessor->ExecutionFunction = [this](FMassEntityManager& ExecuteEntityManager, FMassExecutionContext& Context)
+			{
+				UE_LOG(LogMassComplexTest, Log, TEXT("Begin Creating LargeArchetype")); 
+
+				TArray<FMassEntityHandle> Entities;
+				EntityManager->BatchReserveEntities(NumToSpawn, Entities);
+				for (int32 i = 0; i < NumToSpawn; ++i)
+				{
+					Context.BuildEntity(LargeArchetype, Entities[i]);
+				}
+
+				UE_LOG(LogMassComplexTest, Log, TEXT("End Creating LargeArchetype"));
+			};
+		PhaseManager->RegisterDynamicProcessor(*CreateLargeProcessor);
+
+
+		return true;
+	}
+
+	virtual bool Update() override
+	{
+		Super::Update();
+		return true;
+	}
+
+	virtual void VerifyLatentResults() override
+	{
+		Super::VerifyLatentResults();
+
+		// Verify that entities ended up in IntsArchetype.
+		const int32 IntsCount = EntityManager->DebugGetArchetypeEntitiesCount(IntsArchetype);
+		AITEST_EQUAL_LATENT(TEXT("Expected number of entities created in IntsArchetype"), IntsCount, NumToSpawn);
+
+		const int32 FloatsCount = EntityManager->DebugGetArchetypeEntitiesCount(FloatsArchetype);
+		AITEST_EQUAL_LATENT(TEXT("Expected number of entities created in FloatsArchetype"), FloatsCount, NumToSpawn);
+
+		const int32 LargeCount = EntityManager->DebugGetArchetypeEntitiesCount(LargeArchetype);
+		AITEST_EQUAL_LATENT(TEXT("Expected number of entities created in LargeArchetype"), LargeCount, NumToSpawn);
+
+		PhaseManager->UnregisterDynamicProcessor(*CreateIntsProcessor);
+		CreateIntsProcessor->MarkAsGarbage();
+		CreateIntsProcessor = nullptr;
+		PhaseManager->UnregisterDynamicProcessor(*CreateFloatsProcessor);
+		CreateFloatsProcessor->MarkAsGarbage();
+		CreateFloatsProcessor = nullptr;
+		PhaseManager->UnregisterDynamicProcessor(*CreateLargeProcessor);
+		CreateLargeProcessor->MarkAsGarbage();
+		CreateLargeProcessor = nullptr;
+	}
+};
+IMPLEMENT_AI_LATENT_TEST(FComplex_CreateEntitiesParallel, "System.Mass.Complex.CreateEntitiesParallel");
+
+struct FComplex_CreateEntitiesParallelWithProcessingQueue : public FComplex_CreateEntitiesParallel
+{
+	virtual bool SetUp() override
+	{
+		if (!FComplex_CreateEntitiesParallel::SetUp())
+		{
+			return false;
+		}
+		return true;
+	}
+};
+IMPLEMENT_AI_LATENT_TEST(FComplex_CreateEntitiesParallelWithProcessingQueue, "System.Mass.Complex.CreateEntitiesParallelWithProcessingQueue");
+
+//-----------------------------------------------------------------------------
+// Create Entities during query iteration
+//-----------------------------------------------------------------------------
+struct FComplex_CreateEntitiesDuringQuery : FProcessingPhasesTestBase
+{
+	using Super = FProcessingPhasesTestBase;
+
+	static constexpr int32 NumToSpawn = 10;
+	static constexpr EMassProcessingPhase PhaseToRun = EMassProcessingPhase::PrePhysics;
+
+	TObjectPtr<UMassTestProcessorBase> CreateEntitiesProcessor = nullptr;
+	TObjectPtr<UMassTestProcessorBase> TestObserver = nullptr;
+	int32 NumEntitiesObserved = 0;
+
+	FMassElementBitSet FloatsComposition;
+	FMassElementBitSet LargeComposition;
+	virtual bool PopulatePhasesConfig() override
+	{
+		return true;
+	}
+
+	virtual bool SetUp() override
+	{
+		if (!Super::SetUp())
+		{
+			return false;
+		}
+
+		TestObserver = NewTestProcessor<UMassTestProcessorBase>(EntityManager);
+		TestObserver->EntityQuery.AddRequirement<FTestFragment_Float>(EMassFragmentAccess::ReadOnly);
+		FMassObserverManager& ObserverManager = EntityManager->GetObserverManager();
+		ObserverManager.AddObserverInstance(FTestFragment_Float::StaticStruct(), EMassObservedOperationFlags::CreateEntity, TestObserver);
+		TestObserver->ExecutionFunction = [this](FMassEntityManager& ExecuteEntityManager, FMassExecutionContext& Context)
+			{
+				UE_LOG(LogMassComplexTest, Log, TEXT("Observing Created Entities"));
+
+				TestObserver->EntityQuery.ForEachEntityChunk(Context, [this](FMassExecutionContext& Context)
+					{
+						for (FMassExecutionContext::FEntityIterator EntityIt = Context.CreateEntityIterator(); EntityIt; ++EntityIt)
+						{
+							NumEntitiesObserved++;
+						}
+					});
+
+				UE_LOG(LogMassComplexTest, Log, TEXT("End Observing Created Entities"));
+			};
+
+		TArray<FMassEntityHandle> Created;
+		Created.Reserve(NumToSpawn);
+		EntityManager->BatchCreateEntities(IntsArchetype, NumToSpawn, Created);
+
+		FloatsComposition = FloatsArchetype.GetCompositionBitSetChecked();
+		LargeComposition = LargeArchetype.GetCompositionBitSetChecked();
+
+		CreateEntitiesProcessor = NewObject<UMassTestProcessorBase>(GetTransientPackage());
+		CreateEntitiesProcessor->SetShouldAllowMultipleInstances(true);
+		CreateEntitiesProcessor->SetProcessingPhase(PhaseToRun);
+		CreateEntitiesProcessor->CallInitialize(GetTransientPackage(), EntityManager->AsShared());
+		CreateEntitiesProcessor->EntityQuery.AddRequirement<FTestFragment_Int>(EMassFragmentAccess::ReadWrite);
+		CreateEntitiesProcessor->EntityQuery.AddCreatedArchetype(FloatsComposition);
+		CreateEntitiesProcessor->EntityQuery.AddCreatedArchetype(LargeComposition);
+
+		CreateEntitiesProcessor->ExecutionFunction = [this](FMassEntityManager& ExecuteEntityManager, FMassExecutionContext& Context)
+			{
+				UE_LOG(LogMassComplexTest, Log, TEXT("Begin Creating Entities During Query"));
+
+				CreateEntitiesProcessor->EntityQuery.ForEachEntityChunk(Context, [this](FMassExecutionContext& Context)
+					{
+						for (FMassExecutionContext::FEntityIterator EntityIt = Context.CreateEntityIterator(); EntityIt; ++EntityIt)
+						{
+							Context.CreateEntity(FloatsArchetype);
+							Context.CreateEntity(LargeArchetype);
+						}
+					});
+
+				UE_LOG(LogMassComplexTest, Log, TEXT("End Creating Entities During Query"));
+			};
+		PhaseManager->RegisterDynamicProcessor(*CreateEntitiesProcessor);
+
+		return true;
+	}
+
+	virtual bool Update() override
+	{
+		Super::Update();
+		return true;
+	}
+
+	virtual void VerifyLatentResults() override
+	{
+		Super::VerifyLatentResults();
+
+		// Verify that entities ended up in IntsArchetype.
+		const int32 IntsCount = EntityManager->DebugGetArchetypeEntitiesCount(IntsArchetype);
+		AITEST_EQUAL_LATENT(TEXT("Expected number of entities created in IntsArchetype"), IntsCount, NumToSpawn);
+
+		const int32 FloatsCount = EntityManager->DebugGetArchetypeEntitiesCount(FloatsArchetype);
+		AITEST_EQUAL_LATENT(TEXT("Expected number of entities created in FloatsArchetype"), FloatsCount, NumToSpawn);
+		AITEST_EQUAL_LATENT(TEXT("Expected number of observed entities to match number created"), FloatsCount, NumEntitiesObserved);
+
+		const int32 LargeCount = EntityManager->DebugGetArchetypeEntitiesCount(LargeArchetype);
+		AITEST_EQUAL_LATENT(TEXT("Expected number of entities created in LargeArchetype"), LargeCount, NumToSpawn);
+
+		PhaseManager->UnregisterDynamicProcessor(*CreateEntitiesProcessor);
+		CreateEntitiesProcessor->MarkAsGarbage();
+		CreateEntitiesProcessor = nullptr;
+	}
+};
+IMPLEMENT_AI_LATENT_TEST(FComplex_CreateEntitiesDuringQuery, "System.Mass.Complex.CreateEntitiesDuringQuery");
+
+//-----------------------------------------------------------------------------
+// Deferred Create Entities during query iteration
+//-----------------------------------------------------------------------------
+struct FComplex_DeferredCreateEntitiesDuringQuery : FProcessingPhasesTestBase
+{
+	using Super = FProcessingPhasesTestBase;
+
+	static constexpr int32 NumToSpawn = 50000;
+	static constexpr EMassProcessingPhase PhaseToRun = EMassProcessingPhase::PrePhysics;
+
+	TObjectPtr<UMassTestProcessorBase> CreateEntitiesProcessor = nullptr;
+
+	FMassElementBitSet FloatsComposition;
+	FMassElementBitSet LargeComposition;
+	virtual bool PopulatePhasesConfig() override
+	{
+		return true;
+	}
+
+	virtual bool SetUp() override
+	{
+		if (!Super::SetUp())
+		{
+			return false;
+		}
+
+		TArray<FMassEntityHandle> Created;
+		Created.Reserve(NumToSpawn);
+		EntityManager->BatchCreateEntities(IntsArchetype, NumToSpawn, Created);
+
+		FloatsComposition = FloatsArchetype.GetCompositionBitSetChecked();
+		LargeComposition = LargeArchetype.GetCompositionBitSetChecked();
+
+		CreateEntitiesProcessor = NewObject<UMassTestProcessorBase>(GetTransientPackage());
+		CreateEntitiesProcessor->SetShouldAllowMultipleInstances(true);
+		CreateEntitiesProcessor->SetProcessingPhase(PhaseToRun);
+		CreateEntitiesProcessor->CallInitialize(GetTransientPackage(), EntityManager->AsShared());
+		CreateEntitiesProcessor->EntityQuery.AddRequirement<FTestFragment_Int>(EMassFragmentAccess::ReadWrite);
+
+		CreateEntitiesProcessor->ExecutionFunction = [this](FMassEntityManager& ExecuteEntityManager, FMassExecutionContext& Context)
+			{
+				CreateEntitiesProcessor->EntityQuery.ForEachEntityChunk(Context, [this](FMassExecutionContext& Context)
+					{
+						for (FMassExecutionContext::FEntityIterator EntityIt = Context.CreateEntityIterator(); EntityIt; ++EntityIt)
+						{
+							Context.Defer().PushCommand<FMassDeferredCreateCommand>(
+								[this](FMassEntityManager& DeferredEntityManager)
+								{
+									DeferredEntityManager.CreateEntity(FloatsArchetype);
+								}
+							);
+
+							Context.Defer().PushCommand<FMassDeferredCreateCommand>(
+								[this](FMassEntityManager& DeferredEntityManager)
+								{
+									DeferredEntityManager.CreateEntity(LargeArchetype);
+								}
+							);
+						}
+					});
+			};
+		PhaseManager->RegisterDynamicProcessor(*CreateEntitiesProcessor);
+
+		return true;
+	}
+
+	virtual bool Update() override
+	{
+		Super::Update();
+		return true;
+	}
+
+	virtual void VerifyLatentResults() override
+	{
+		Super::VerifyLatentResults();
+
+		// Verify that entities ended up in IntsArchetype.
+		const int32 IntsCount = EntityManager->DebugGetArchetypeEntitiesCount(IntsArchetype);
+		AITEST_EQUAL_LATENT(TEXT("Expected number of entities created in IntsArchetype"), IntsCount, NumToSpawn);
+
+		const int32 FloatsCount = EntityManager->DebugGetArchetypeEntitiesCount(FloatsArchetype);
+		AITEST_EQUAL_LATENT(TEXT("Expected number of entities created in FloatsArchetype"), FloatsCount, NumToSpawn);
+
+		const int32 LargeCount = EntityManager->DebugGetArchetypeEntitiesCount(LargeArchetype);
+		AITEST_EQUAL_LATENT(TEXT("Expected number of entities created in LargeArchetype"), LargeCount, NumToSpawn);
+
+		PhaseManager->UnregisterDynamicProcessor(*CreateEntitiesProcessor);
+		CreateEntitiesProcessor->MarkAsGarbage();
+		CreateEntitiesProcessor = nullptr;
+	}
+};
+IMPLEMENT_AI_LATENT_TEST(FComplex_DeferredCreateEntitiesDuringQuery, "System.Mass.Complex.DeferredCreateEntitiesDuringQuery");
+
+
+//-----------------------------------------------------------------------------
+// Batched entity creation in scheduled processors
+//-----------------------------------------------------------------------------
+struct FComplex_BatchedCreateEntitiesParallel : FProcessingPhasesTestBase
+{
+	using Super = FProcessingPhasesTestBase;
+
+	static constexpr int32 NumToSpawn = 1000000;
+	static constexpr EMassProcessingPhase PhaseToRun = EMassProcessingPhase::PrePhysics;
+
+	TObjectPtr<UMassTestProcessorBase> CreateIntsProcessor = nullptr;
+	TObjectPtr<UMassTestProcessorBase> CreateFloatsProcessor = nullptr;
+	TObjectPtr<UMassTestProcessorBase> CreateLargeProcessor = nullptr;
+
+	FMassElementBitSet IntsComposition;
+	FMassElementBitSet FloatsComposition;
+	FMassElementBitSet LargeComposition;
+
+	virtual bool PopulatePhasesConfig() override
+	{
+		return true;
+	}
+
+	virtual bool SetUp() override
+	{
+		if (!Super::SetUp())
+		{
+			return false;
+		}
+
+		IntsComposition = IntsArchetype.GetCompositionBitSetChecked();
+		CreateIntsProcessor = NewObject<UMassTestProcessorBase>(GetTransientPackage());
+		CreateIntsProcessor->SetShouldAllowMultipleInstances(true);
+		CreateIntsProcessor->SetProcessingPhase(PhaseToRun);
+		CreateIntsProcessor->GetEntityCreationRequirements().AddCreatedArchetype(IntsComposition);
+		CreateIntsProcessor->ExecutionFunction = [this](FMassEntityManager& ExecuteEntityManager, FMassExecutionContext& Context)
+			{
+				UE_LOG(LogMassComplexTest, Log, TEXT("Begin Creating IntsArchetype"));
+
+				TArray<FMassEntityHandle> Created;
+				Created.Reserve(NumToSpawn);
+				EntityManager->BatchReserveEntities(NumToSpawn, Created);
+				for (int32 i = 0; i < NumToSpawn; ++i)
+				{
+					Context.BuildEntity(IntsArchetype, Created[i]);
+				}
+
+				UE_LOG(LogMassComplexTest, Log, TEXT("End Creating IntsArchetype"));
+			};
+		PhaseManager->RegisterDynamicProcessor(*CreateIntsProcessor);
+
+		FloatsComposition = FloatsArchetype.GetCompositionBitSetChecked();
+		CreateFloatsProcessor = NewObject<UMassTestProcessorBase>(GetTransientPackage());
+		CreateFloatsProcessor->SetShouldAllowMultipleInstances(true);
+		CreateFloatsProcessor->SetProcessingPhase(PhaseToRun);
+		CreateFloatsProcessor->GetEntityCreationRequirements().AddCreatedArchetype(FloatsComposition);
+		CreateFloatsProcessor->ExecutionFunction = [this](FMassEntityManager& ExecuteEntityManager, FMassExecutionContext& Context)
+			{
+				UE_LOG(LogMassComplexTest, Log, TEXT("Begin Creating FloatsArchetype"));
+
+				TArray<FMassEntityHandle> Created;
+				Created.Reserve(NumToSpawn);
+				EntityManager->BatchReserveEntities(NumToSpawn, Created);
+				for (int32 i = 0; i < NumToSpawn; ++i)
+				{
+					Context.BuildEntity(FloatsArchetype, Created[i]);
+				}
+
+				UE_LOG(LogMassComplexTest, Log, TEXT("End Creating FloatsArchetype"));
+			};
+		PhaseManager->RegisterDynamicProcessor(*CreateFloatsProcessor);
+
+		LargeComposition = LargeArchetype.GetCompositionBitSetChecked();
+		CreateLargeProcessor = NewObject<UMassTestProcessorBase>(GetTransientPackage());
+		CreateLargeProcessor->SetShouldAllowMultipleInstances(true);
+		CreateLargeProcessor->SetProcessingPhase(PhaseToRun);
+		CreateLargeProcessor->GetEntityCreationRequirements().AddCreatedArchetype(LargeComposition);
+		CreateLargeProcessor->ExecutionFunction = [this](FMassEntityManager& ExecuteEntityManager, FMassExecutionContext& Context)
+			{
+				UE_LOG(LogMassComplexTest, Log, TEXT("Begin Creating LargeArchetype"));
+
+				TArray<FMassEntityHandle> Created;
+				Created.Reserve(NumToSpawn);
+				EntityManager->BatchReserveEntities(NumToSpawn, Created);
+				for (int32 i = 0; i < NumToSpawn; ++i)
+				{
+					Context.BuildEntity(LargeArchetype, Created[i]);
+				}
+
+				UE_LOG(LogMassComplexTest, Log, TEXT("End Creating LargeArchetype"));
+			};
+		PhaseManager->RegisterDynamicProcessor(*CreateLargeProcessor);
+
+
+		return true;
+	}
+
+	virtual bool Update() override
+	{
+		Super::Update();
+		return true;
+	}
+
+	virtual void VerifyLatentResults() override
+	{
+		Super::VerifyLatentResults();
+
+		// Verify that entities ended up in IntsArchetype.
+		const int32 IntsCount = EntityManager->DebugGetArchetypeEntitiesCount(IntsArchetype);
+		AITEST_EQUAL_LATENT(TEXT("Expected number of entities created in IntsArchetype"), IntsCount, NumToSpawn);
+
+		const int32 FloatsCount = EntityManager->DebugGetArchetypeEntitiesCount(FloatsArchetype);
+		AITEST_EQUAL_LATENT(TEXT("Expected number of entities created in FloatsArchetype"), FloatsCount, NumToSpawn);
+
+		const int32 LargeCount = EntityManager->DebugGetArchetypeEntitiesCount(LargeArchetype);
+		AITEST_EQUAL_LATENT(TEXT("Expected number of entities created in LargeArchetype"), LargeCount, NumToSpawn);
+
+		PhaseManager->UnregisterDynamicProcessor(*CreateIntsProcessor);
+		CreateIntsProcessor->MarkAsGarbage();
+		CreateIntsProcessor = nullptr;
+		PhaseManager->UnregisterDynamicProcessor(*CreateFloatsProcessor);
+		CreateFloatsProcessor->MarkAsGarbage();
+		CreateFloatsProcessor = nullptr;
+		PhaseManager->UnregisterDynamicProcessor(*CreateLargeProcessor);
+		CreateLargeProcessor->MarkAsGarbage();
+		CreateLargeProcessor = nullptr;
+	}
+};
+IMPLEMENT_AI_LATENT_TEST(FComplex_BatchedCreateEntitiesParallel, "System.Mass.Complex.BatchedCreateEntitiesParallel");
+
+//-----------------------------------------------------------------------------
+// Deferred entity creation to compare with immediate in processor API
+//-----------------------------------------------------------------------------
+struct FComplex_CreateEntitiesDeferred : FProcessingPhasesTestBase
+{
+	using Super = FProcessingPhasesTestBase;
+
+	static constexpr int32 NumToSpawn = 1000000;
+	static constexpr EMassProcessingPhase PhaseToRun = EMassProcessingPhase::PrePhysics;
+
+	TObjectPtr<UMassTestProcessorBase> CreateIntsProcessor = nullptr;
+	TObjectPtr<UMassTestProcessorBase> CreateFloatsProcessor = nullptr;
+	TObjectPtr<UMassTestProcessorBase> CreateLargeProcessor = nullptr;
+
+	// Keep the descriptors (used to build archetypes) and also store handles.
+	FMassArchetypeCompositionDescriptor IntsComposition;
+	FMassArchetypeCompositionDescriptor FloatsComposition;
+	FMassArchetypeCompositionDescriptor LargeComposition;
+
+	FMassArchetypeHandle IntsArchetypeHandle;
+	FMassArchetypeHandle FloatsArchetypeHandle;
+	FMassArchetypeHandle LargeArchetypeHandle;
+
+	virtual bool PopulatePhasesConfig() override
+	{
+		return true;
+	}
+
+	virtual bool SetUp() override
+	{
+		if (!Super::SetUp())
+		{
+			return false;
+		}
+
+		CreateIntsProcessor = NewObject<UMassTestProcessorBase>(GetTransientPackage());
+		CreateIntsProcessor->SetShouldAllowMultipleInstances(true);
+		CreateIntsProcessor->SetProcessingPhase(PhaseToRun);
+		CreateIntsProcessor->ExecutionFunction = [this](FMassEntityManager& ExecuteEntityManager, FMassExecutionContext& Context)
+			{
+				Context.Defer().PushCommand<FMassDeferredCreateCommand>(
+					[this](FMassEntityManager& DeferredEntityManager)
+					{
+						UE_LOG(LogMassComplexTest, Log, TEXT("Begin Creating IntsArchetype"));
+						for (int32 i = 0; i < NumToSpawn; ++i)
+						{
+							DeferredEntityManager.CreateEntity(IntsArchetype);
+						}
+						UE_LOG(LogMassComplexTest, Log, TEXT("End Creating IntsArchetype"));
+					}
+				);
+			};
+		PhaseManager->RegisterDynamicProcessor(*CreateIntsProcessor);
+
+		CreateFloatsProcessor = NewObject<UMassTestProcessorBase>(GetTransientPackage());
+		CreateFloatsProcessor->SetShouldAllowMultipleInstances(true);
+		CreateFloatsProcessor->SetProcessingPhase(PhaseToRun);
+		CreateFloatsProcessor->ExecutionFunction = [this](FMassEntityManager& ExecuteEntityManager, FMassExecutionContext& Context)
+			{
+				Context.Defer().PushCommand<FMassDeferredCreateCommand>(
+					[this](FMassEntityManager& DeferredEntityManager)
+					{
+						UE_LOG(LogMassComplexTest, Log, TEXT("Begin Creating FloatsArchetype"));
+						for (int32 i = 0; i < NumToSpawn; ++i)
+						{
+							DeferredEntityManager.CreateEntity(FloatsArchetype);
+						}
+						UE_LOG(LogMassComplexTest, Log, TEXT("End Creating FloatsArchetype"));
+					}
+				);
+			};
+		PhaseManager->RegisterDynamicProcessor(*CreateFloatsProcessor);
+
+		CreateLargeProcessor = NewObject<UMassTestProcessorBase>(GetTransientPackage());
+		CreateLargeProcessor->SetShouldAllowMultipleInstances(true);
+		CreateLargeProcessor->SetProcessingPhase(PhaseToRun);
+		CreateLargeProcessor->ExecutionFunction = [this](FMassEntityManager& ExecuteEntityManager, FMassExecutionContext& Context)
+			{
+				for (int32 i = 0; i < NumToSpawn; ++i)
+				{
+					Context.Defer().PushCommand<FMassDeferredCreateCommand>(
+						[this](FMassEntityManager& DeferredEntityManager)
+						{
+							DeferredEntityManager.CreateEntity(LargeArchetype);
+						}
+					);
+				}
+			};
+		PhaseManager->RegisterDynamicProcessor(*CreateLargeProcessor);
+
+		return true;
+	}
+
+	virtual bool Update() override
+	{
+		Super::Update();
+		return true;
+	}
+
+	virtual void VerifyLatentResults() override
+	{
+		Super::VerifyLatentResults();
+
+		// After the phase runs, deferred commands are flushed; verify counts.
+		const int32 IntsCount = EntityManager->DebugGetArchetypeEntitiesCount(IntsArchetype);
+		AITEST_EQUAL_LATENT(TEXT("Expected number of entities created in IntsArchetype"), IntsCount, NumToSpawn);
+
+		const int32 FloatsCount = EntityManager->DebugGetArchetypeEntitiesCount(FloatsArchetype);
+		AITEST_EQUAL_LATENT(TEXT("Expected number of entities created in FloatsArchetype"), FloatsCount, NumToSpawn);
+
+		const int32 LargeCount = EntityManager->DebugGetArchetypeEntitiesCount(LargeArchetype);
+		AITEST_EQUAL_LATENT(TEXT("Expected number of entities created in LargeArchetype"), LargeCount, NumToSpawn);
+
+		PhaseManager->UnregisterDynamicProcessor(*CreateIntsProcessor);
+		CreateIntsProcessor->MarkAsGarbage();
+		CreateIntsProcessor = nullptr;
+		PhaseManager->UnregisterDynamicProcessor(*CreateFloatsProcessor);
+		CreateFloatsProcessor->MarkAsGarbage();
+		CreateFloatsProcessor = nullptr;
+		PhaseManager->UnregisterDynamicProcessor(*CreateLargeProcessor);
+		CreateLargeProcessor->MarkAsGarbage();
+		CreateLargeProcessor = nullptr;
+	}
+};
+IMPLEMENT_AI_LATENT_TEST(FComplex_CreateEntitiesDeferred, "System.Mass.Complex.CreateEntitiesDeferred");
+
+#include "HAL/IConsoleManager.h"
+
+struct FComplex_DestroyEntitiesDuringQuery : FProcessingPhasesTestBase
+{
+	using Super = FProcessingPhasesTestBase;
+
+	static constexpr int32 NumToSpawn = 10;
+	static constexpr EMassProcessingPhase PhaseToRun = EMassProcessingPhase::PrePhysics;
+
+	TObjectPtr<UMassTestProcessorBase> DestroyEntitiesProcessor = nullptr;
+
+	// Store the old CVar state to restore it later
+	bool bOldFullyParallel = true;
+
+	virtual bool PopulatePhasesConfig() override
+	{
+		return true;
+	}
+
+	virtual bool SetUp() override
+	{
+		// 1. Force Mass to run synchronously via CVar before Super::SetUp initializes the PhaseManager.
+		if (IConsoleVariable* CVar = IConsoleManager::Get().FindConsoleVariable(TEXT("mass.FullyParallel")))
+		{
+			bOldFullyParallel = CVar->GetBool();
+			CVar->Set(false, ECVF_SetByCode);
+		}
+
+		// Now initialize the environment (which reads the CVar)
+		if (!Super::SetUp())
+		{
+			return false;
+		}
+
+		// 2. Spawn entities
+		TArray<FMassEntityHandle> Created;
+		Created.Reserve(NumToSpawn);
+		EntityManager->BatchCreateEntities(IntsArchetype, NumToSpawn, Created);
+
+		// 3. Setup the processor
+		DestroyEntitiesProcessor = NewObject<UMassTestProcessorBase>(GetTransientPackage());
+		DestroyEntitiesProcessor->SetShouldAllowMultipleInstances(true);
+		DestroyEntitiesProcessor->SetProcessingPhase(PhaseToRun);
+		DestroyEntitiesProcessor->CallInitialize(GetTransientPackage(), EntityManager->AsShared());
+
+		// Note: We are NOT setting bRequiresMainThreadExecution = true here.
+		// We rely on "mass.FullyParallel = false" to force the phase to run on the GameThread.
+
+		DestroyEntitiesProcessor->EntityQuery.AddRequirement<FTestFragment_Int>(EMassFragmentAccess::ReadWrite);
+
+		DestroyEntitiesProcessor->ExecutionFunction = [this](FMassEntityManager& ExecuteEntityManager, FMassExecutionContext& Context)
+			{
+				UE_LOG(LogMassComplexTest, Log, TEXT("Begin Destroying Entities (CVar Synchronous)"));
+
+				DestroyEntitiesProcessor->EntityQuery.ForEachEntityChunk(Context, [this](FMassExecutionContext& Context)
+					{
+						for (FMassExecutionContext::FEntityIterator EntityIt = Context.CreateEntityIterator(); EntityIt; ++EntityIt)
+						{
+							FMassEntityHandle Entity = Context.GetEntity(EntityIt);
+
+							// Direct access to EntityManager (Main Buffer).
+							// This is safe because mass.FullyParallel=false ensures this runs on the GameThread.
+							EntityManager->Defer().DestroyEntity(Entity);
+						}
+					});
+
+				UE_LOG(LogMassComplexTest, Log, TEXT("End Destroying Entities (CVar Synchronous)"));
+			};
+
+		PhaseManager->RegisterDynamicProcessor(*DestroyEntitiesProcessor);
+
+		return true;
+	}
+
+	virtual bool Update() override
+	{
+		Super::Update();
+		return true;
+	}
+
+	virtual void VerifyLatentResults() override
+	{
+		Super::VerifyLatentResults();
+
+		const int32 IntsCount = EntityManager->DebugGetArchetypeEntitiesCount(IntsArchetype);
+
+		AITEST_EQUAL_LATENT(TEXT("Expected IntsArchetype count to be 0 after synchronous destruction"), IntsCount, 0);
+
+		PhaseManager->UnregisterDynamicProcessor(*DestroyEntitiesProcessor);
+		DestroyEntitiesProcessor->MarkAsGarbage();
+		DestroyEntitiesProcessor = nullptr;
+
+		// Restore the CVar to avoid side effects on other tests
+		if (IConsoleVariable* CVar = IConsoleManager::Get().FindConsoleVariable(TEXT("mass.FullyParallel")))
+		{
+			CVar->Set(bOldFullyParallel, ECVF_SetByCode);
+		}
+	}
+};
+IMPLEMENT_AI_LATENT_TEST(FComplex_DestroyEntitiesDuringQuery, "System.Mass.Complex.DestroyEntitiesDuringQuery");
 } // namespace FMassComplexTest
 
 UE_ENABLE_OPTIMIZATION_SHIP
diff --git a/Engine/Source/Developer/MassEntityTestSuite/Private/MassProcessingPhasesTest.cpp b/Engine/Source/Developer/MassEntityTestSuite/Private/MassProcessingPhasesTest.cpp
--- a/Engine/Source/Developer/MassEntityTestSuite/Private/MassProcessingPhasesTest.cpp
+++ b/Engine/Source/Developer/MassEntityTestSuite/Private/MassProcessingPhasesTest.cpp
@@ -26,7 +26,15 @@ bool FProcessingPhasesTestBase::SetUp()
 	{
 		PhaseManager = MakeShareable(new UE::Mass::Testing::FMassTestProcessingPhaseManager());
 
-		EntityManager->Initialize();
+		FMassEntityManagerStorageInitParams InitializationParams;
+#if WITH_MASS_CONCURRENT_RESERVE
+
+		InitializationParams.Emplace<FMassEntityManager_InitParams_Concurrent>();
+#else
+		InitializationParams.Emplace<FMassEntityManager_InitParams_SingleThreaded>();
+#endif // WITH_MASS_CONCURRENT_RESERVE
+
+		EntityManager->Initialize(InitializationParams);
 
 		if (PopulatePhasesConfig())
 		{
diff --git a/Engine/Source/Developer/MassEntityTestSuite/Private/MassTestTypes.cpp b/Engine/Source/Developer/MassEntityTestSuite/Private/MassTestTypes.cpp
--- a/Engine/Source/Developer/MassEntityTestSuite/Private/MassTestTypes.cpp
+++ b/Engine/Source/Developer/MassEntityTestSuite/Private/MassTestTypes.cpp
@@ -18,7 +18,16 @@ bool FExecutionTestBase::SetUp()
 {
 	EntityManager = MakeShareable(new FMassEntityManager(bMakeWorldEntityManagersOwner ? FAITestHelpers::GetWorld() : nullptr));
 	EntityManager->SetDebugName(TEXT("MassEntityTestSuite"));
-	EntityManager->Initialize();
+
+	FMassEntityManagerStorageInitParams InitializationParams;
+#if WITH_MASS_CONCURRENT_RESERVE
+
+	InitializationParams.Emplace<FMassEntityManager_InitParams_Concurrent>();
+#else
+	InitializationParams.Emplace<FMassEntityManager_InitParams_SingleThreaded>();
+#endif // WITH_MASS_CONCURRENT_RESERVE
+
+	EntityManager->Initialize(InitializationParams);
 	EntityManager->PostInitialize();
 
 	return true;
@@ -37,12 +46,13 @@ bool FEntityTestBase::SetUp()
 	FExecutionTestBase::SetUp();
 	check(EntityManager);
 
-	const UScriptStruct* FragmentTypes[] = { FTestFragment_Float::StaticStruct(), FTestFragment_Int::StaticStruct() };
+	const UScriptStruct* FragmentTypes[] = { FTestFragment_Float::StaticStruct(), FTestFragment_Int::StaticStruct(), FTestFragment_Large::StaticStruct()};
 
 	EmptyArchetype = EntityManager->CreateArchetype(MakeArrayView<const UScriptStruct*>(nullptr, 0));
 	FloatsArchetype = EntityManager->CreateArchetype(MakeArrayView(&FragmentTypes[0], 1));
 	IntsArchetype = EntityManager->CreateArchetype(MakeArrayView(&FragmentTypes[1], 1));
 	FloatsIntsArchetype = EntityManager->CreateArchetype(MakeArrayView(FragmentTypes, 2));
+	LargeArchetype = EntityManager->CreateArchetype(MakeArrayView(&FragmentTypes[2], 1));
 
 	FTestFragment_Int IntFrag;
 	IntFrag.Value = FTestFragment_Int::TestIntValue;
diff --git a/Engine/Source/Developer/MassEntityTestSuite/Public/MassEntityTestTypes.h b/Engine/Source/Developer/MassEntityTestSuite/Public/MassEntityTestTypes.h
--- a/Engine/Source/Developer/MassEntityTestSuite/Public/MassEntityTestTypes.h
+++ b/Engine/Source/Developer/MassEntityTestSuite/Public/MassEntityTestTypes.h
@@ -287,6 +287,11 @@ public:
 	UE_API void SetUseParallelForEachEntityChunk(bool bEnable);
 	UE_API void TestExecute(TSharedPtr<FMassEntityManager>& EntityManager);
 
+	FMassEntityCreationRequirements& GetEntityCreationRequirements()
+	{
+		return ProcessorEntityCreationRequirements;
+	}
+
 	/** public on purpose, this is a test processor, no worries about access*/
 	FMassEntityQuery EntityQuery;
 };
@@ -430,6 +435,7 @@ struct FEntityTestBase : FExecutionTestBase
 	FMassArchetypeHandle FloatsArchetype;
 	FMassArchetypeHandle IntsArchetype;
 	FMassArchetypeHandle FloatsIntsArchetype;
+	FMassArchetypeHandle LargeArchetype;
 
 	FInstancedStruct InstanceInt;
 
diff --git a/Engine/Source/Runtime/MassEntity/Private/MassArchetypeTypes.cpp b/Engine/Source/Runtime/MassEntity/Private/MassArchetypeTypes.cpp
--- a/Engine/Source/Runtime/MassEntity/Private/MassArchetypeTypes.cpp
+++ b/Engine/Source/Runtime/MassEntity/Private/MassArchetypeTypes.cpp
@@ -15,6 +15,12 @@ uint32 GetTypeHash(const FMassArchetypeHandle& Instance)
 	return GetTypeHash(Instance.DataPtr.Get());
 }
 
+const FMassElementBitSet& FMassArchetypeHandle::GetCompositionBitSetChecked() const
+{
+	check(DataPtr);
+	return DataPtr.Get()->GetCompositionBitSet();
+}
+
 //-----------------------------------------------------------------------------
 // FMassArchetypeVersionedHandle
 //-----------------------------------------------------------------------------
diff --git a/Engine/Source/Runtime/MassEntity/Private/MassEntityManager.cpp b/Engine/Source/Runtime/MassEntity/Private/MassEntityManager.cpp
--- a/Engine/Source/Runtime/MassEntity/Private/MassEntityManager.cpp
+++ b/Engine/Source/Runtime/MassEntity/Private/MassEntityManager.cpp
@@ -612,6 +612,30 @@ FMassEntityHandle FMassEntityManager::CreateEntity(TConstArrayView<FInstancedStr
 	return Entity;
 }
 
+FMassEntityHandle FMassEntityManager::CreateEntityAsyncUnsafe(const FMassArchetypeHandle& ArchetypeHandle, const FMassArchetypeSharedFragmentValues& SharedFragmentValues)
+{
+	check(ArchetypeHandle.IsValid());
+
+	MASS_BREAKPOINT(UE::Mass::Debug::FBreakpoint::CheckCreateEntityBreakpoints(ArchetypeHandle));
+
+	const FMassEntityHandle Entity = ReserveEntity();
+	InternalBuildEntity(Entity
+		, GetOrCreateSuitableArchetype(ArchetypeHandle, SharedFragmentValues.GetBitSet())
+		, SharedFragmentValues);
+	return Entity;
+}
+
+void FMassEntityManager::BuildEntityAsyncUnsafe(const FMassArchetypeHandle& ArchetypeHandle, const FMassEntityHandle& ReservedEntityHandle, const FMassArchetypeSharedFragmentValues& SharedFragmentValues)
+{
+	check(ArchetypeHandle.IsValid());
+
+	MASS_BREAKPOINT(UE::Mass::Debug::FBreakpoint::CheckCreateEntityBreakpoints(ArchetypeHandle));
+
+	InternalBuildEntity(ReservedEntityHandle
+		, GetOrCreateSuitableArchetype(ArchetypeHandle, SharedFragmentValues.GetBitSet())
+		, SharedFragmentValues);
+}
+
 FMassEntityHandle FMassEntityManager::ReserveEntity()
 {
 	FMassEntityHandle Result = GetEntityStorageInterface().AcquireOne();
@@ -765,7 +789,7 @@ TSharedRef<FMassEntityManager::FEntityCreationContext> FMassEntityManager::Batch
 
 	return InternalBatchCreateReservedEntities(
 		GetOrCreateSuitableArchetype(ArchetypeHandle, SharedFragmentValues.GetBitSet())
-		, SharedFragmentValues, ReservedEntities);
+		, SharedFragmentValues, ReservedEntities, /*bCreateContext=*/true).ToSharedRef();
 }
 
 TSharedRef<FMassEntityManager::FEntityCreationContext> FMassEntityManager::BatchCreateEntities(const FMassArchetypeHandle& ArchetypeHandle
@@ -780,14 +804,29 @@ TSharedRef<FMassEntityManager::FEntityCreationContext> FMassEntityManager::Batch
 	MASS_BREAKPOINT(UE::Mass::Debug::FBreakpoint::CheckCreateEntityBreakpoints(ArchetypeHandle));
 
 	TConstArrayView<FMassEntityHandle> ReservedEntities = BatchReserveEntities(Count, InOutEntities);
-	
+
 	return InternalBatchCreateReservedEntities(
 		GetOrCreateSuitableArchetype(ArchetypeHandle, SharedFragmentValues.GetBitSet())
-		, SharedFragmentValues, ReservedEntities);
+		, SharedFragmentValues, ReservedEntities, /*bCreateContext=*/true).ToSharedRef();
 }
 
-TSharedRef<FMassEntityManager::FEntityCreationContext> FMassEntityManager::InternalBatchCreateReservedEntities(const FMassArchetypeHandle& ArchetypeHandle
-	, const FMassArchetypeSharedFragmentValues& SharedFragmentValues, TConstArrayView<FMassEntityHandle> ReservedEntities)
+void FMassEntityManager::BatchCreateEntitiesAsyncUnsafe(const FMassArchetypeHandle& ArchetypeHandle
+	, const FMassArchetypeSharedFragmentValues& SharedFragmentValues, const int32 Count, TArray<FMassEntityHandle>& InOutEntities)
+{
+	TRACE_CPUPROFILER_EVENT_SCOPE(Mass_BatchCreateEntitiesAsyncUnsafe);
+
+	testableCheckfReturn(ArchetypeHandle.IsValid(), return, TEXT("%hs expecting a valid ArchetypeHandle"), __FUNCTION__);
+
+	MASS_BREAKPOINT(UE::Mass::Debug::FBreakpoint::CheckCreateEntityBreakpoints(ArchetypeHandle));
+
+	TConstArrayView<FMassEntityHandle> ReservedEntities = BatchReserveEntities(Count, InOutEntities);
+
+	InternalBatchCreateReservedEntities(ArchetypeHandle, SharedFragmentValues, ReservedEntities, /*bCreateContext=*/false);
+}
+
+TSharedPtr<FMassEntityManager::FEntityCreationContext> FMassEntityManager::InternalBatchCreateReservedEntities(const FMassArchetypeHandle& ArchetypeHandle
+	, const FMassArchetypeSharedFragmentValues& SharedFragmentValues, TConstArrayView<FMassEntityHandle> ReservedEntities
+	, const bool bCreateContext)
 {
 	// Functions calling into this one are required to verify that the archetype handle is valid
 	FMassArchetypeData* ArchetypeData = FMassArchetypeHelper::ArchetypeDataFromHandle(ArchetypeHandle);
@@ -800,7 +839,7 @@ TSharedRef<FMassEntityManager::FEntityCreationContext> FMassEntityManager::Inter
 		check(IsEntityValid(Entity));
 		const UE::Mass::IEntityStorageInterface::EEntityState EntityState = GetEntityStorageInterface().GetEntityState(Entity.Index);
 		checkf(EntityState == UE::Mass::IEntityStorageInterface::EEntityState::Reserved, TEXT("Trying to build entities that are not reserved. Check all handles are reserved or consider using BatchCreateEntities"));
-		
+
 		GetEntityStorageInterface().SetArchetypeFromShared(Entity.Index, ArchetypeHandle.DataPtr);
 	}
 
@@ -809,7 +848,12 @@ TSharedRef<FMassEntityManager::FEntityCreationContext> FMassEntityManager::Inter
 
 	UE_TRACE_MASS_ENTITIES_CREATED(ReservedEntities, *ArchetypeData)
 
-	return ObserverManager.GetOrMakeCreationContext(ReservedEntities, FMassArchetypeEntityCollection(ArchetypeHandle, MoveTemp(TargetArchetypeEntityRanges)));
+		if (bCreateContext)
+		{
+			return ObserverManager.GetOrMakeCreationContext(ReservedEntities, FMassArchetypeEntityCollection(ArchetypeHandle, MoveTemp(TargetArchetypeEntityRanges)));
+		}
+
+	return nullptr;
 }
 
 void FMassEntityManager::DestroyEntity(FMassEntityHandle Entity)
@@ -1152,7 +1196,6 @@ void FMassEntityManager::InternalBuildEntity(FMassEntityHandle Entity, const FMa
 	NewArchetype->AddEntity(Entity, SharedFragmentValues);
 	
 	UE_TRACE_MASS_ENTITY_CREATED(Entity, *NewArchetype)
-
 	ObserverManager.OnPostEntityCreated(Entity, NewArchetype->GetCompositionDescriptor());
 }
 
diff --git a/Engine/Source/Runtime/MassEntity/Private/MassEntityQuery.cpp b/Engine/Source/Runtime/MassEntity/Private/MassEntityQuery.cpp
--- a/Engine/Source/Runtime/MassEntity/Private/MassEntityQuery.cpp
+++ b/Engine/Source/Runtime/MassEntity/Private/MassEntityQuery.cpp
@@ -817,6 +817,7 @@ void FMassEntityQuery::ExportRequirements(FMassExecutionRequirements& OutRequire
 {
 	FMassSubsystemRequirements::ExportRequirements(OutRequirements);
 	FMassFragmentRequirements::ExportRequirements(OutRequirements);
+	EntityCreationRequirements.ExportRequirements(OutRequirements);
 }
 
 void FMassEntityQuery::DebugEnableEntityOwnerLogging()
diff --git a/Engine/Source/Runtime/MassEntity/Private/MassEntityTypes.cpp b/Engine/Source/Runtime/MassEntity/Private/MassEntityTypes.cpp
--- a/Engine/Source/Runtime/MassEntity/Private/MassEntityTypes.cpp
+++ b/Engine/Source/Runtime/MassEntity/Private/MassEntityTypes.cpp
@@ -3,6 +3,7 @@
 #include "MassEntityTypes.h"
 
 #include "MassArchetypeData.h"
+#include "MassArchetypeTypes.h"
 
 #include UE_INLINE_GENERATED_CPP_BY_NAME(MassEntityTypes)
 
@@ -36,6 +37,12 @@ FString LexToString(const EMassObservedOperationFlags Value)
 //-----------------------------------------------------------------------------
 // FMassArchetypeCompositionDescriptor
 //-----------------------------------------------------------------------------
+FMassArchetypeCompositionDescriptor::FMassArchetypeCompositionDescriptor(const FMassArchetypeHandle& ArchetypeHandle)
+{
+	FMassArchetypeData& ArchetypeData = FMassArchetypeHelper::ArchetypeDataFromHandleChecked(ArchetypeHandle);
+	*this = ArchetypeData.GetCompositionDescriptor();
+}
+
 PRAGMA_DISABLE_DEPRECATION_WARNINGS
 int32 FMassArchetypeCompositionDescriptor::CountStoredTypes() const
 {
diff --git a/Engine/Source/Runtime/MassEntity/Private/MassExecutionContext.cpp b/Engine/Source/Runtime/MassEntity/Private/MassExecutionContext.cpp
--- a/Engine/Source/Runtime/MassEntity/Private/MassExecutionContext.cpp
+++ b/Engine/Source/Runtime/MassEntity/Private/MassExecutionContext.cpp
@@ -23,7 +23,7 @@ FMassExecutionContext::FMassExecutionContext(FMassEntityManager& InEntityManager
 FMassExecutionContext::FMassExecutionContext(const FMassExecutionContext& Other)
 	: EntityManager(Other.EntityManager)
 {
-	// we're using assignment operator here as a setup helper, as per operator=='s comment.
+	// we're using assignment operator here as a setup helper, as per operator='s comment.
 	*this = Other;
 
 	QueriesStack.Reset();
@@ -152,8 +152,8 @@ void FMassExecutionContext::BindLinkedEntity(const FMassEntityLinkFragment& Enti
 {
 	check(QueriesStack.Num());
 	checkfSlow(
-		SharedFragmentViews.ContainsByPredicate(
-			[&EntityLinkFragment](const FSharedFragmentView& Element)
+		ConstSharedFragmentViews.ContainsByPredicate(
+			[&EntityLinkFragment](const FConstSharedFragmentView& Element)
 			{
 				const void* ElemMem = Element.FragmentView.GetMemory();
 				return ElemMem == static_cast<const void*>(&EntityLinkFragment);
@@ -199,6 +199,9 @@ void FMassExecutionContext::PushQuery(FMassEntityQuery& InQuery)
 	}
 #endif // WITH_MASSENTITY_DEBUG
 
+	RuntimeData.PreviousEntityCreationRequirements = ActiveEntityCreationRequirements;
+	ActiveEntityCreationRequirements = &InQuery.GetEntityCreationRequirements();
+
 	RuntimeData.IteratorSerialNumber = ++IteratorSerialNumberGenerator;
 }
 
@@ -209,6 +212,8 @@ void FMassExecutionContext::PopQuery(const FMassEntityQuery& InQuery)
 
 	SetSubsystemRequirementBits(RuntimeData.ConstSubsystemsBitSet, RuntimeData.MutableSubsystemsBitSet);
 
+	ActiveEntityCreationRequirements = RuntimeData.PreviousEntityCreationRequirements;
+
 	QueriesStack.RemoveAt(QueriesStack.Num() - 1, EAllowShrinking::No);
 }
 
@@ -222,6 +227,42 @@ FMassExecutionContext::FEntityIterator FMassExecutionContext::CreateEntityIterat
 	return FEntityIterator(*this, QueriesStack.Last());
 }
 
+FMassEntityHandle FMassExecutionContext::CreateEntity(const FMassArchetypeHandle& ArchetypeHandle,
+	const FMassArchetypeSharedFragmentValues& SharedValues)
+{
+	checkf(ActiveEntityCreationRequirements != nullptr, TEXT("Attempting to create an entity without declaring FMassEntityCreationRequirements"));
+#if WITH_MASSENTITY_DEBUG
+	// this check is a bit slow and very unlikely to fail for anyone other than the programmer implementing the processors calling this
+	const FMassArchetypeData& ArchetypeData = FMassArchetypeHelper::ArchetypeDataFromHandleChecked(ArchetypeHandle);
+	checkf(ActiveEntityCreationRequirements->Contains(ArchetypeData.GetCompositionBitSet()), TEXT("Declared FMassEntityCreationRequirements do not allow creating entites of the requested archetype"));
+#endif
+	FMassEntityHandle EntityHandle = EntityManager->CreateEntityAsyncUnsafe(ArchetypeHandle, SharedValues);
+
+	return EntityHandle;
+}
+
+void FMassExecutionContext::BuildEntity(const FMassArchetypeHandle& ArchetypeHandle, const FMassEntityHandle& EntityHandle, const FMassArchetypeSharedFragmentValues& SharedValues)
+{
+	checkf(ActiveEntityCreationRequirements != nullptr, TEXT("Attempting to create an entity without declaring EntityCreationRequirements"));
+#if WITH_MASSENTITY_DEBUG
+	const FMassArchetypeData& ArchetypeData = FMassArchetypeHelper::ArchetypeDataFromHandleChecked(ArchetypeHandle);
+	checkf(ActiveEntityCreationRequirements->Contains(ArchetypeData.GetCompositionBitSet()), TEXT("Declared FMassEntityCreationRequirements do not allow creating entites of the requested archetype"));
+#endif
+	EntityManager->BuildEntityAsyncUnsafe(ArchetypeHandle, EntityHandle, SharedValues);
+}
+
+void FMassExecutionContext::BatchCreateEntities(const FMassArchetypeHandle& ArchetypeHandle
+	, const FMassArchetypeSharedFragmentValues& SharedFragmentValues, const int32 Count, TArray<FMassEntityHandle>& InOutEntities)
+{
+	checkf(ActiveEntityCreationRequirements != nullptr, TEXT("Attempting to create an entity without declaring FMassEntityCreationRequirements"));
+#if WITH_MASSENTITY_DEBUG
+	const FMassArchetypeData& ArchetypeData = FMassArchetypeHelper::ArchetypeDataFromHandleChecked(ArchetypeHandle);
+	checkf(ActiveEntityCreationRequirements->Contains(ArchetypeData.GetCompositionBitSet()), TEXT("Declared FMassEntityCreationRequirements do not allow creating entites of the requested archetype"));
+#endif
+
+	EntityManager->BatchCreateEntitiesAsyncUnsafe(ArchetypeHandle, SharedFragmentValues, Count, InOutEntities);
+}
+
 FMassExecutionContext& FMassExecutionContext::GetDummyInstance()
 {
 	static FMassExecutionContext DummyContext(*TSharedRef<FMassEntityManager>(MakeShareable(new FMassEntityManager())));
diff --git a/Engine/Source/Runtime/MassEntity/Private/MassProcessor.cpp b/Engine/Source/Runtime/MassEntity/Private/MassProcessor.cpp
--- a/Engine/Source/Runtime/MassEntity/Private/MassProcessor.cpp
+++ b/Engine/Source/Runtime/MassEntity/Private/MassProcessor.cpp
@@ -14,6 +14,7 @@
 
 DECLARE_CYCLE_STAT(TEXT("MassProcessor Group Completed"), Mass_GroupCompletedTask, STATGROUP_TaskGraphTasks);
 DECLARE_CYCLE_STAT(TEXT("Mass Processor Task"), STAT_Mass_DoTask, STATGROUP_Mass);
+DECLARE_CYCLE_STAT(TEXT("Mass Processor Execute"), STAT_Mass_ExecuteProcessor, STATGROUP_Mass);
 
 #if WITH_MASSENTITY_DEBUG
 namespace UE::Mass::Debug
@@ -169,8 +170,10 @@ void UMassProcessor::CallInitialize(const TNotNull<UObject*> Owner, const TShare
 
 		ConfigureQueries(EntityManager);
 
+		ResolveRequestedArchetypes(EntityManager.Get());
+
 		bool bNeedsGameThread = ProcessorRequirements.DoesRequireGameThreadExecution();
-		for (const FMassEntityQuery* QueryPtr : OwnedQueries)
+		for (FMassEntityQuery* QueryPtr : OwnedQueries)
 		{
 			CA_ASSUME(QueryPtr);
 			bNeedsGameThread = (bNeedsGameThread || QueryPtr->DoesRequireGameThreadExecution());
@@ -189,6 +192,44 @@ void UMassProcessor::CallInitialize(const TNotNull<UObject*> Owner, const TShare
 	}
 }
 
+void UMassProcessor::ResolveRequestedArchetypes(const FMassEntityManager& EntityManager)
+{
+	if (bCreatedArchetypeHandlesResolved)
+	{
+		return;
+	}
+
+	ProcessorEntityCreationRequirements.ResolveArchetypes(EntityManager);
+
+	for (FMassEntityQuery* QueryPtr : OwnedQueries)
+	{
+		CA_ASSUME(QueryPtr);
+
+		FMassEntityCreationRequirements& QueryEntityCreationRequirements = QueryPtr->GetEntityCreationRequirements();
+
+		QueryEntityCreationRequirements.ResolveArchetypes(EntityManager);
+
+		bQueryHasEntityCreationRequirements = bQueryHasEntityCreationRequirements || !QueryEntityCreationRequirements.IsEmpty();
+	}
+	 
+	bCreatedArchetypeHandlesResolved = true;
+}
+
+void UMassProcessor::ClearArchetypesCreatedHandles()
+{
+	bCreatedArchetypeHandlesResolved = false;
+	ProcessorEntityCreationRequirements.ClearCreatedArchetypeHandles();
+
+	for (FMassEntityQuery* QueryPtr : OwnedQueries)
+	{
+		CA_ASSUME(QueryPtr);
+
+		FMassEntityCreationRequirements& QueryEntityCreationRequirements = QueryPtr->GetEntityCreationRequirements();
+
+		QueryEntityCreationRequirements.ClearCreatedArchetypeHandles();
+	}
+}
+
 void UMassProcessor::InitializeInternal(UObject&, const TSharedRef<FMassEntityManager>&)
 {
 	// empty in base class
@@ -267,6 +308,7 @@ void UMassProcessor::PostInitProperties()
 
 void UMassProcessor::CallExecute(FMassEntityManager& EntityManager, FMassExecutionContext& Context)
 {
+	SCOPE_CYCLE_COUNTER(STAT_Mass_ExecuteProcessor);
 	if (LIKELY(ensureMsgf(IsActive(), TEXT("Trying to CallExecute for an inactive processor %s"), *GetProcessorName())))
 	{
 		TRACE_CPUPROFILER_EVENT_SCOPE_TEXT(*StatId);
@@ -297,9 +339,12 @@ void UMassProcessor::CallExecute(FMassEntityManager& EntityManager, FMassExecuti
 
 		if (bAllRequirementsMet)
 		{
+			Context.SetActiveArchetypeRequirements(&ProcessorEntityCreationRequirements);
 			Execute(EntityManager, Context);
 		}
 
+		Context.SetActiveArchetypeRequirements(nullptr);
+
 		if (ActivationState == EActivationState::OneShot)
 		{
 			MakeInactive();
@@ -344,6 +389,23 @@ void UMassProcessor::ExportRequirements(FMassExecutionRequirements& OutRequireme
 		CA_ASSUME(Query);
 		Query->ExportRequirements(OutRequirements);
 	}
+
+	ProcessorEntityCreationRequirements.ExportRequirements(OutRequirements);
+}
+
+bool UMassProcessor::HasEntityCreationRequirements()
+{
+	return bQueryHasEntityCreationRequirements || !ProcessorEntityCreationRequirements.IsEmpty();
+}
+
+void UMassProcessor::ExportArchetypesCreated(TArray<FMassArchetypeHandle>& OutArchetypes)
+{
+	OutArchetypes.Append(ProcessorEntityCreationRequirements.GetCreatedArchetypes());
+	for (FMassEntityQuery* Query : OwnedQueries)
+	{
+		CA_ASSUME(Query);
+		OutArchetypes.Append(Query->GetEntityCreationRequirements().GetCreatedArchetypes());
+	}
 }
 
 void UMassProcessor::RegisterQuery(FMassEntityQuery& Query)
diff --git a/Engine/Source/Runtime/MassEntity/Private/MassProcessorDependencySolver.cpp b/Engine/Source/Runtime/MassEntity/Private/MassProcessorDependencySolver.cpp
--- a/Engine/Source/Runtime/MassEntity/Private/MassProcessorDependencySolver.cpp
+++ b/Engine/Source/Runtime/MassEntity/Private/MassProcessorDependencySolver.cpp
@@ -945,6 +945,14 @@ void FMassProcessorDependencySolver::ResolveDependencies(TArray<FMassProcessorOr
 			// this line is a part of a nice trick we're doing here utilizing EntityManager's archetype creation based on 
 			// what each processor expects, and EntityQuery's capability to cache archetypes matching its requirements (used below)
 			EntityManager->CreateArchetype(AllNodes[ProcessorNodeIndex].Requirements.AsCompositionDescriptor());
+
+			if (AllNodes[ProcessorNodeIndex].Processor->HasEntityCreationRequirements())
+			{
+				AllNodes[ProcessorNodeIndex].Processor->ClearArchetypesCreatedHandles();
+
+				// ResolveRequestedArchetypes creates the actual archetypes
+				AllNodes[ProcessorNodeIndex].Processor->ResolveRequestedArchetypes(*EntityManager.Get());
+			}
 		}
 	}
 
@@ -965,11 +973,18 @@ void FMassProcessorDependencySolver::ResolveDependencies(TArray<FMassProcessorOr
 			// that will let the user know their processor is misconfigured.
 			// We do collect archetype information for the processors that never get pruned because we're
 			// using this information for the dependency calculations, regardless of ShouldAllowQueryBasedPruning
-			if (bDoQueryBasedPruning || Node.Processor->GetOwnedQueriesNum())
+			if (bDoQueryBasedPruning || Node.Processor->GetOwnedQueriesNum() || Node.Processor->HasEntityCreationRequirements())
 			{
 				// for each processor-representing node we cache information on which archetypes among the once we've created 
 				// above (see the EntityManager.CreateArchetype call in the previous loop) match this processor. 
 				Node.Processor->GetArchetypesMatchingOwnedQueries(*EntityManager.Get(), Node.ValidArchetypes);
+				Node.Processor->ExportArchetypesCreated(Node.ValidArchetypes);
+			}
+
+			if (bCreateVirtualArchetypes)
+			{
+				// reset the archetype handles if they're virtual
+				Node.Processor->ClearArchetypesCreatedHandles();
 			}
 
 			// prune the archetype-less processors
diff --git a/Engine/Source/Runtime/MassEntity/Private/MassRequirements.cpp b/Engine/Source/Runtime/MassEntity/Private/MassRequirements.cpp
--- a/Engine/Source/Runtime/MassEntity/Private/MassRequirements.cpp
+++ b/Engine/Source/Runtime/MassEntity/Private/MassRequirements.cpp
@@ -337,6 +337,39 @@ void FMassFragmentRequirements::Reset()
 	// to add elements to it. This "requirements" instance is now "empty" but still valid
 }
 
+//-----------------------------------------------------------------------------
+// FMassEntityCreationRequirements
+//-----------------------------------------------------------------------------
+void FMassEntityCreationRequirements::ExportRequirements(FMassExecutionRequirements& Out) const
+{
+	FMassElementBitSet ElementsBitSet;
+	for (const FMassElementBitSet& Composition : Compositions)
+	{
+		ElementsBitSet += Composition;
+	}
+
+	Out.Fragments.Write += ElementsBitSet.Get<FMassFragmentBitSet>();
+	Out.ChunkFragments.Write += ElementsBitSet.Get<FMassChunkFragmentBitSet>();
+	Out.SharedFragments.Write += ElementsBitSet.Get<FMassSharedFragmentBitSet>();
+
+	// Note: ConstSharedFragments are "read-only" by nature; declaring write on them
+	// would be misleading. We still need *archetype-level* exclusivity, which the solver
+	// enforces via ValidArchetypes overlap. So no write on const shared here.
+	Out.RequiredAllTags += ElementsBitSet.Get<FMassTagBitSet>();
+
+	Out.ResourcesUsedCount = INDEX_NONE; // force recomputation
+}
+
+void FMassEntityCreationRequirements::ResolveArchetypes(const FMassEntityManager& EntityManager)
+{
+	ArchetypeHandles.Reset(Compositions.Num());
+
+	for (const FMassElementBitSet& Composition : Compositions)
+	{
+		ArchetypeHandles.Add(const_cast<FMassEntityManager&>(EntityManager).CreateArchetype(FMassArchetypeCompositionDescriptor(Composition)));
+	}
+}
+
 //-----------------------------------------------------------------------------
 // DEPRECATED
 //-----------------------------------------------------------------------------
diff --git a/Engine/Source/Runtime/MassEntity/Public/MassArchetypeTypes.h b/Engine/Source/Runtime/MassEntity/Public/MassArchetypeTypes.h
--- a/Engine/Source/Runtime/MassEntity/Public/MassArchetypeTypes.h
+++ b/Engine/Source/Runtime/MassEntity/Public/MassArchetypeTypes.h
@@ -47,6 +47,11 @@ struct FMassArchetypeHandle final
 
 	void Reset();
 
+	/**
+	 * Returns the composition bitset for the archetype, asserts if archetype is invalid
+	 */
+	MASSENTITY_API const FMassElementBitSet& GetCompositionBitSetChecked() const;
+
 private:
 	FMassArchetypeHandle(const TSharedPtr<FMassArchetypeData>& InDataPtr);
 
diff --git a/Engine/Source/Runtime/MassEntity/Public/MassEntityManager.h b/Engine/Source/Runtime/MassEntity/Public/MassEntityManager.h
--- a/Engine/Source/Runtime/MassEntity/Public/MassEntityManager.h
+++ b/Engine/Source/Runtime/MassEntity/Public/MassEntityManager.h
@@ -97,6 +97,7 @@ struct FMassEntityManager : public TSharedFromThis<FMassEntityManager>, public F
 {
 	friend FMassEntityQuery;
 	friend FMassDebugger;
+	friend FMassExecutionContext; // Execution context needs direct access to unsafe api.
 
 	DECLARE_MULTICAST_DELEGATE_OneParam(FOnNewArchetypeDelegate, const FMassArchetypeHandle&);
 
@@ -366,14 +367,14 @@ public:
 	/**
 	 * Assigns all entities indicated by Collections to a given archetype group.
 	 * Note that depending on their individual composition each entity can end up in a different archetype.
-	 * @paramm GroupHandle indicates the target group. Passing an invalid group handle will get logged as warning and ignored.
+	 * @param GroupHandle indicates the target group. Passing an invalid group handle will get logged as warning and ignored.
 	 */
 	UE_API void BatchGroupEntities(const UE::Mass::FArchetypeGroupHandle GroupHandle, TConstArrayView<FMassArchetypeEntityCollection> Collections);
 
 	/**
 	 * Assigns all entities indicated by InEntities to a given archetype group.
 	 * Note that depending on their individual composition each entity can end up in a different archetype.
-	 * @paramm GroupHandle indicates the target group. Passing an invalid group handle will get logged as warning and ignored.
+	 * @param GroupHandle indicates the target group. Passing an invalid group handle will get logged as warning and ignored.
 	 */
 	UE_API void BatchGroupEntities(const UE::Mass::FArchetypeGroupHandle GroupHandle, TConstArrayView<FMassEntityHandle> InEntities);
 
@@ -771,6 +772,41 @@ private:
 		return ConstSharedFragmentsContainer.FindOrAdd<T>(Hash, Forward<TArgs>(InArgs)...);
 	}
 
+	/**
+	 * Creates fully built entity ready to be used by the subsystem
+	 * Does not check for synchronous execution and should only be used by processors through the execution context
+	 * This should only ever be used by FMassExecutionContext which will validate data access safety.
+	 * @param ArchetypeHandle the archetype you want this entity to be
+	 * @param SharedFragmentValues to be associated with the entity
+	 * @return FMassEntityHandle id of the newly created entity
+	 */
+	UE_API FMassEntityHandle CreateEntityAsyncUnsafe(const FMassArchetypeHandle& ArchetypeHandle, const FMassArchetypeSharedFragmentValues& SharedFragmentValues = {});
+	
+	/**
+	 * Fully buils an entity with a reserved handle ready to be used by the subsystem
+	 * Does not check for synchronous execution and should only be used by processors through the execution context
+	 * This should only ever be used by FMassExecutionContext which will validate data access safety.
+	 * @param ArchetypeHandle the archetype you want this entity to be
+	 * @param ReservedEntityHandle the already reserved handle to use for the new enetity
+	 * @param SharedFragmentValues to be associated with the entity
+	 */
+	UE_API void BuildEntityAsyncUnsafe(const FMassArchetypeHandle& ArchetypeHandle, const FMassEntityHandle& ReservedEntityHandle, const FMassArchetypeSharedFragmentValues& SharedFragmentValues = {});
+
+	/**
+	 * A version of CreateEntity that's creating a number of entities (Count) in one go
+	 * Unsafe version only for use in processors via ExecutionContext
+	 * Note: this does not create or return a FEntityCreationContext because it's assumed that observers are already locked by the processing phase.
+	 * @param ArchetypeHandle you want this entity to be
+	 * @param SharedFragmentValues to be associated with the entities
+	 * @param Count number of entities to create
+	 * @param InOutEntities the newly created entities are appended to given array, i.e. the pre-existing content of OutEntities won't be affected by the call
+	 */
+	UE_API void BatchCreateEntitiesAsyncUnsafe(const FMassArchetypeHandle& ArchetypeHandle, const FMassArchetypeSharedFragmentValues& SharedFragmentValues, const int32 Count, TArray<FMassEntityHandle>& InOutEntities);
+	void BatchCreateEntitiesAsyncUnsafe(const FMassArchetypeHandle& ArchetypeHandle, const int32 Count, TArray<FMassEntityHandle>& InOutEntities)
+	{
+		BatchCreateEntitiesAsyncUnsafe(ArchetypeHandle, FMassArchetypeSharedFragmentValues(), Count, InOutEntities);
+	}
+
 public:
 
 #if WITH_EDITOR || WITH_MASSENTITY_DEBUG
@@ -1179,8 +1215,9 @@ private:
 	/** Note that it's the caller's responsibility to ensure `SharedFragmentType` is a kind of FMassSharedFragment */
 	UE_API const FSharedStruct* InternalGetSharedFragmentPtr(FMassEntityHandle EntityHandle, const UScriptStruct* SharedFragmentType) const;
 
-	TSharedRef<FEntityCreationContext> InternalBatchCreateReservedEntities(const FMassArchetypeHandle& ArchetypeHandle,
-		const FMassArchetypeSharedFragmentValues& SharedFragmentValues, TConstArrayView<FMassEntityHandle> ReservedEntities);
+	TSharedPtr<FMassEntityManager::FEntityCreationContext> InternalBatchCreateReservedEntities(const FMassArchetypeHandle& ArchetypeHandle
+		, const FMassArchetypeSharedFragmentValues& SharedFragmentValues, TConstArrayView<FMassEntityHandle> ReservedEntities
+		, const bool bCreateContext);
 
 	UE::Mass::FStorageType& GetEntityStorageInterface() const
 	{
diff --git a/Engine/Source/Runtime/MassEntity/Public/MassEntityQuery.h b/Engine/Source/Runtime/MassEntity/Public/MassEntityQuery.h
--- a/Engine/Source/Runtime/MassEntity/Public/MassEntityQuery.h
+++ b/Engine/Source/Runtime/MassEntity/Public/MassEntityQuery.h
@@ -244,6 +244,21 @@ public:
 	 */
 	UE_API void DebugEnableEntityOwnerLogging();
 
+	void AddCreatedArchetype(const FMassElementBitSet& InComposition)
+	{
+		EntityCreationRequirements.AddCreatedArchetype(InComposition);
+	}
+
+	void AddCreatedArchetype(FMassElementBitSet&& InComposition)
+	{
+		EntityCreationRequirements.AddCreatedArchetype(InComposition);
+	}
+
+	FMassEntityCreationRequirements& GetEntityCreationRequirements()
+	{
+		return EntityCreationRequirements;
+	}
+
 private:
 	/**
 	 * Incrementally sorts all ValidArchetypes to fill OrderedArchetypeIndices with the expected order of archetype processing.
@@ -283,6 +298,8 @@ private:
 	TArray<FArchetypeGroupingStep> GroupSortingSteps;
 	TArray<TArray<UE::Mass::FArchetypeGroupID>> CachedGroupIDs;
 
+	FMassEntityCreationRequirements EntityCreationRequirements;
+
 	/** 
 	 * Controls whether ParallelForEachEntityChunk created dedicated command buffer for each job. This is required 
 	 * to ensure thread safety. Disable by calling SetParallelCommandBufferEnabled(false) if execution function doesn't 
diff --git a/Engine/Source/Runtime/MassEntity/Public/MassEntityTypes.h b/Engine/Source/Runtime/MassEntity/Public/MassEntityTypes.h
--- a/Engine/Source/Runtime/MassEntity/Public/MassEntityTypes.h
+++ b/Engine/Source/Runtime/MassEntity/Public/MassEntityTypes.h
@@ -30,6 +30,7 @@ DECLARE_CYCLE_STAT_EXTERN(TEXT("Mass Total Frame Time"), STAT_Mass_Total, STATGR
 
 struct FMassArchetypeData;
 struct FMassEntityQuery;
+struct FMassArchetypeHandle;
 
 namespace UE::Mass
 {
@@ -123,6 +124,8 @@ struct FMassArchetypeCompositionDescriptor
 		: ElementsBitSet(InElements)
 	{}
 
+	explicit MASSENTITY_API FMassArchetypeCompositionDescriptor(const FMassArchetypeHandle& ArchetypeHandle);
+
 	void Reset()
 	{
 		ElementsBitSet.Reset();
diff --git a/Engine/Source/Runtime/MassEntity/Public/MassExecutionContext.h b/Engine/Source/Runtime/MassEntity/Public/MassExecutionContext.h
--- a/Engine/Source/Runtime/MassEntity/Public/MassExecutionContext.h
+++ b/Engine/Source/Runtime/MassEntity/Public/MassExecutionContext.h
@@ -12,6 +12,7 @@
 #include "MassArchetypeTypes.h"
 #include "MassSubsystemAccess.h"
 #include "MassProcessor.h"
+#include "MassEntityCollection.h"
 #include "MassEntityView.h"
 
 
@@ -87,6 +88,9 @@ private:
 		FMassExternalSubsystemBitSet ConstSubsystemsBitSet;
 		FMassExternalSubsystemBitSet MutableSubsystemsBitSet;
 
+		/** used to restore creation requirements on query pop, storing the previous value to easily restore processor creation requirements */
+		const FMassEntityCreationRequirements* PreviousEntityCreationRequirements = nullptr;
+
 #if WITH_MASSENTITY_DEBUG
 		/** MaxBreakFragmentCount needs to be bigger than the greatest number of fragments a query has a write requirement for to that can have a breakpoint set */
 		static constexpr uint32 MaxFragmentBreakpointCount = 8;
@@ -96,7 +100,7 @@ private:
 		int32 BreakFragmentsCount = 0;
 #endif // WITH_MASSENTITY_DEBUG
 
-		/** Serial number to ensure iterator consistency (subsequent calls to CreateEntityIterator should not pass equivelncy test) */
+		/** Serial number to ensure iterator consistency (subsequent calls to CreateEntityIterator should not pass equivalency test) */
 		uint32 IteratorSerialNumber = 0;
 
 		/** Helper function to create an empty instance with a valid Query ptr */
@@ -121,6 +125,9 @@ private:
 	 * is mainly utilized to avoid numerous small flushes in favor of fewer larger ones. */
 	bool bFlushDeferredCommands = true;
 
+	/** Active permissions for archetype creation (processor-level when no query is active; query-level otherwise) */
+	const FMassEntityCreationRequirements* ActiveEntityCreationRequirements = nullptr;
+
 	TArrayView<FFragmentView> GetMutableRequirements() { return FragmentViews; }
 	TArrayView<FChunkFragmentView> GetMutableChunkRequirements() { return ChunkFragmentViews; }
 	TArrayView<FConstSharedFragmentView> GetMutableConstSharedRequirements() { return ConstSharedFragmentViews; }
@@ -504,7 +511,7 @@ public:
 	T* GetLinkedEntityFragmentPtr()
 	{
 		checkf(QueriesStack.Num(), TEXT("There must be an active query to access fragments on the linked entity"));
-		checkf(LinkedEntityView.IsValid(), TEXT("No entity has been linked. Add a FMassLinkedEntityFragment to entities in this query to access fragments on a linked entity that isn't in this query."));
+		checkf(LinkedEntityView.IsValid(), TEXT("No entity has been linked. Add a FMassEntityLinkFragment to entities in this query to access fragments on a linked entity that isn't in this query."));
 		checkfSlow(QueriesStack.Last().Query->HasMutableIndirectEntityRequirement<T>()
 			, TEXT("Attempting to access a fragment from a linked entity without first adding the LinkedEntityFragment requirement to the query"));
 
@@ -515,7 +522,7 @@ public:
 	const T* GetLinkedEntityConstFragmentPtr()
 	{
 		checkf(QueriesStack.Num(), TEXT("There must be an active query to access fragments on the linked entity"));
-		checkf(LinkedEntityView.IsValid(), TEXT("No entity has been linked. Add a FMassLinkedEntityFragment to entities in this query to access fragments on a linked entity that isn't in this query."));
+		checkf(LinkedEntityView.IsValid(), TEXT("No entity has been linked. Add a FMassEntityLinkFragment to entities in this query to access fragments on a linked entity that isn't in this query."));
 		checkfSlow(QueriesStack.Last().Query->HasConstIndirectEntityRequirement<T>()
 			, TEXT("Attempting to access a fragment from a linked entity without first adding the LinkedEntityFragment requirement to the query"));
 
@@ -625,6 +632,27 @@ public:
 		CurrentArchetypeCompositionBitSet = CompositionBitSet;
 	}
 
+	/**
+	* Validates creation request by the active FMassEntityCreationRequirements and creates an entity
+	* of the declared archetype. Will assert on validation failure. 
+	*/
+	MASSENTITY_API FMassEntityHandle CreateEntity(const FMassArchetypeHandle& ArchetypeHandle, const FMassArchetypeSharedFragmentValues& SharedFragmentValues = {});
+	MASSENTITY_API void BuildEntity(const FMassArchetypeHandle& ArchetypeHandle, const FMassEntityHandle& EntityHandle, const FMassArchetypeSharedFragmentValues& SharedFragmentValues = {});
+
+	MASSENTITY_API void BatchCreateEntities(const FMassArchetypeHandle& ArchetypeHandle,
+		const FMassArchetypeSharedFragmentValues& SharedFragmentValues, const int32 Count, TArray<FMassEntityHandle>& InOutEntities);
+	void BatchCreateEntities(const FMassArchetypeHandle& ArchetypeHandle, const int32 Count, TArray<FMassEntityHandle>& InOutEntities)
+	{
+		BatchCreateEntities(ArchetypeHandle, FMassArchetypeSharedFragmentValues(), Count, InOutEntities);
+
+	}
+
+	/** Set/restore by processor and query push/pop */
+	void SetActiveArchetypeRequirements(const FMassEntityCreationRequirements* InRequirements)
+	{
+		ActiveEntityCreationRequirements = InRequirements;
+	}
+
 #if WITH_MASSENTITY_DEBUG
 	void DebugSetColor(FColor InColor)
 	{
diff --git a/Engine/Source/Runtime/MassEntity/Public/MassProcessor.h b/Engine/Source/Runtime/MassEntity/Public/MassProcessor.h
--- a/Engine/Source/Runtime/MassEntity/Public/MassProcessor.h
+++ b/Engine/Source/Runtime/MassEntity/Public/MassProcessor.h
@@ -184,6 +184,27 @@ public:
 	UE_API void GetArchetypesMatchingOwnedQueries(const FMassEntityManager& EntityManager, TArray<FMassArchetypeHandle>& OutArchetype);
 	UE_API bool DoesAnyArchetypeMatchOwnedQueries(const FMassEntityManager& EntityManager);
 	int32 GetOwnedQueriesNum() const;
+
+	/**
+	 * Indicates weather or not this processor has declared in its requirements that it creates entities.
+	 */
+	UE_API bool HasEntityCreationRequirements();
+
+	/**
+	 * Export the cached archetypes (by handle) that this processor can create
+	 */
+	UE_API void ExportArchetypesCreated(TArray<FMassArchetypeHandle>& OutArchetypes);
+
+	/**
+	 * Resolves all the created archetype handles and creates the archetypes if they don't exist
+	 */
+	UE_API void ResolveRequestedArchetypes(const FMassEntityManager& EntityManager);
+
+	/**
+	 * Clears the archetypes created cache to force it to rebuild before next use
+	 */
+	UE_API void ClearArchetypesCreatedHandles();
+
 	
 #if CPUPROFILERTRACE_ENABLED
 	FString StatId;
@@ -256,6 +277,12 @@ private:
 	/** Used to track whether Initialized has been called. */
 	uint8 bInitialized : 1 = false;
 
+	/** Cache presence of entity creation requirements in owned queries. */
+	uint8 bQueryHasEntityCreationRequirements : 1 = false;
+
+	/** Tracks if created archetypes handles are resolved*/
+	uint8 bCreatedArchetypeHandlesResolved : 1 = false;
+
 	/**
 	 * Processors can be activated/deactivated at runtime. Deactivating a running processor will not disrupt the processing
 	 * graph since the disabled processor's dependencies will get passed down to the subsequent processors depending on this one.
@@ -292,6 +319,9 @@ protected:
 	/** A query representing elements this processor is accessing in Execute function outside of query execution */
 	FMassSubsystemRequirements ProcessorRequirements;
 
+	/** Archetypes this processor can create when no query is on the stack */
+	FMassEntityCreationRequirements ProcessorEntityCreationRequirements;
+
 	/** A QueryExecutor that can optionally be run in lieu of overriding the Execute function. */
 	TSharedPtr<UE::Mass::FQueryExecutor> AutoExecuteQuery;
 
diff --git a/Engine/Source/Runtime/MassEntity/Public/MassRequirements.h b/Engine/Source/Runtime/MassEntity/Public/MassRequirements.h
--- a/Engine/Source/Runtime/MassEntity/Public/MassRequirements.h
+++ b/Engine/Source/Runtime/MassEntity/Public/MassRequirements.h
@@ -5,6 +5,8 @@
 #include "MassEntityTypes.h"
 #include "MassExternalSubsystemTraits.h"
 #include "MassEntityLinkFragments.h"
+#include "MassEntityTypes.h"
+#include "MassArchetypeTypes.h"
 #include "Templates/SubclassOf.h"
 #include "MassRequirements.generated.h"
 
@@ -760,6 +762,67 @@ public:
 	MASSENTITY_API FMassFragmentRequirements(TConstArrayView<const UScriptStruct*> InitList);
 };
 
+/**
+ * Specifies entity archetypes that can be created.
+ * Does NOT imply a filtering constraint; this is purely an access declaration.
+ */
+struct FMassEntityCreationRequirements
+{
+public:
+	/** Add by exact composition */
+	void AddCreatedArchetype(const FMassElementBitSet& InComposition)
+	{
+		Compositions.Add(InComposition);
+	}
+
+	void AddCreatedArchetype(FMassElementBitSet&& InComposition)
+	{
+		Compositions.Add(MoveTemp(InComposition));
+	}
+
+	/** Export as "write" access to all fragment types contained in the declared archetypes */
+	MASSENTITY_API void ExportRequirements(FMassExecutionRequirements& OutRequirements) const;
+
+	/** Populates ArchetypeHandles with all created entity archetypes (EntityManager creates those that don't exist) */
+	MASSENTITY_API void ResolveArchetypes(const FMassEntityManager& EntityManager);
+
+	/** Gets the cached array of archetype handles allowed */
+	const TArray<FMassArchetypeHandle>& GetCreatedArchetypes()
+	{
+		return ArchetypeHandles;
+	}
+
+	/** Returns true if the required composition list containsthe provided composition */
+	bool Contains(const FMassElementBitSet& Composition) const
+	{
+		return Compositions.ContainsByPredicate([&Composition](const FMassElementBitSet& ElementBitSet)
+			{
+				return ElementBitSet.IsEquivalent(Composition);
+			});
+	}
+
+	/** Clear the list of created archetype handles to force an update */
+	void ClearCreatedArchetypeHandles()
+	{
+		ArchetypeHandles.Reset();
+	}
+
+	void Reset()
+	{
+		ArchetypeHandles.Reset();
+		Compositions.Reset();
+	}
+
+	bool IsEmpty() const
+	{
+		return Compositions.Num() == 0;
+	}
+
+private:
+	TArray<FMassArchetypeHandle> ArchetypeHandles;
+	TArray<FMassElementBitSet> Compositions;
+};
+
 //-----------------------------------------------------------------------------
 // INLINE
 //-----------------------------------------------------------------------------
