[Mass] Fixed `System.Mass.Query.OnlyAbsentFragments` unit test broken in last revision (parallel entity creation)

#rnx

#Synced-CL 50097741

diff --git a/Engine/Source/Developer/MassEntityTestSuite/Private/MassQueryTest.cpp b/Engine/Source/Developer/MassEntityTestSuite/Private/MassQueryTest.cpp
--- a/Engine/Source/Developer/MassEntityTestSuite/Private/MassQueryTest.cpp
+++ b/Engine/Source/Developer/MassEntityTestSuite/Private/MassQueryTest.cpp
@@ -328,14 +328,14 @@ struct FQueryTest_OnlyAbsentFragments : FEntityTestBase
 		AITEST_TRUE("Single negative requirement is valid", Query.CheckValidity());
 
 		Query.AddRequirement<FTestFragment_Float>(EMassFragmentAccess::ReadOnly, EMassFragmentPresence::None);
-		Query.AddRequirement<FTestFragment_Bool>(EMassFragmentAccess::ReadOnly, EMassFragmentPresence::None);
+		Query.AddRequirement<FTestFragment_Large>(EMassFragmentAccess::ReadOnly, EMassFragmentPresence::None);
 		AITEST_TRUE("Multiple negative requirement is valid", Query.CheckValidity());
 		
 		Query.CacheArchetypes();
 		AITEST_EQUAL("There's only one default test archetype matching the query", Query.GetArchetypes().Num(), 1);
 		AITEST_TRUE("Only the Empty archetype matches the query", Query.GetArchetypes()[0] == EmptyArchetype);
 
-		const FMassArchetypeHandle NewMatchingArchetypeHandle = EntityManager->CreateArchetype({ FTestFragment_Large::StaticStruct() });
+		const FMassArchetypeHandle NewMatchingArchetypeHandle = EntityManager->CreateArchetype({ FTestFragment_Bool::StaticStruct() });
 		Query.CacheArchetypes();
 		AITEST_EQUAL("The number of matching queries matches expectations", Query.GetArchetypes().Num(), 2);
 		AITEST_TRUE("The new archetype matches the query", Query.GetArchetypes()[1] == NewMatchingArchetypeHandle);
