[Mass] Sparse Elements along with support in queries, commands and entity builder.

Sparse Elements are mass elements (tags and fragments) that can be added to entities without causing them to change archetypes (i.e. sparse elements are not part of archetypes' composition and are individually handled per entity).

#Synced-CL 50439895

diff --git a/Engine/Source/Developer/MassEntityTestSuite/Private/MassSparseElementsTest.cpp b/Engine/Source/Developer/MassEntityTestSuite/Private/MassSparseElementsTest.cpp
new file mode 100644
--- /dev/null
+++ b/Engine/Source/Developer/MassEntityTestSuite/Private/MassSparseElementsTest.cpp
@@ -0,0 +1,656 @@
+// Copyright Epic Games, Inc. All Rights Reserved.
+
+#include "AITestsCommon.h"
+#include "MassArchetypeData.h"
+#include "UObject/Class.h"
+#include "MassEntityBuilder.h"
+#include "MassEntityManager.h"
+#include "MassEntityMacros.h"
+#include "MassEntityTestTypes.h"
+#include "MassEntityView.h"
+#include "MassExecutionContext.h"
+#include "MassExecutor.h"
+
+#define LOCTEXT_NAMESPACE "MassTest"
+
+UE_DISABLE_OPTIMIZATION_SHIP
+
+// to be enabled once this functionality gets implemented
+#define WITH_ARCHETYPE_UPDATING_BITSETS_ON_REMOVE 0
+
+namespace UE::Mass::Test::SparseElements
+{
+#if WITH_MASSENTITY_DEBUG
+	TNotNull<UScriptStruct*> MakeElement(FName TypeName, TNotNull<UScriptStruct*> BaseElementType)
+	{
+		UScriptStruct* NewScriptStruct = NewObject<UScriptStruct>(UScriptStruct::StaticClass(), FName(*FString::Printf(TEXT("%s"), *TypeName.ToString())), RF_Public);
+		NewScriptStruct->SetSuperStruct(BaseElementType);
+		NewScriptStruct->Bind();
+		NewScriptStruct->PrepareCppStructOps();
+		NewScriptStruct->StaticLink(true);
+		NewScriptStruct->AddToRoot();
+		return NewScriptStruct;
+	}
+
+	struct FAutoCleanupOperations
+	{
+		~FAutoCleanupOperations()
+		{
+			for (UScriptStruct* Type : CreatedTypes)
+			{
+				Type->RemoveFromRoot();
+			}
+		}
+
+		TArray<UScriptStruct*> CreatedTypes;
+	};
+
+	struct FTagTestOperations : FAutoCleanupOperations
+	{
+		static bool HasElement(const FMassArchetypeHandle& ArchetypeHandle, TNotNull<UScriptStruct*> SparseElementType, bool bIncludeTransient)
+		{
+			return FMassArchetypeHelper::ArchetypeDataFromHandleChecked(ArchetypeHandle).HasTagType(SparseElementType, bIncludeTransient);
+		}
+
+		static bool HasElement(FMassEntityView EntityView, TNotNull<UScriptStruct*> SparseElementType)
+		{
+			return EntityView.HasTag(SparseElementType);
+		}
+
+		TNotNull<UScriptStruct*> MakeSparseElement()
+		{
+			CreatedTypes.Add(MakeElement("SparseTagTestA", FMassSparseTag::StaticStruct()));
+			return CreatedTypes.Last();
+		}
+
+		TNotNull<UScriptStruct*> MakeRegularElement()
+		{
+			CreatedTypes.Add(MakeElement("RegularTagTestA", FMassTag::StaticStruct()));
+			return CreatedTypes.Last();
+		}
+	};
+
+	struct FFragmentTestOperations : FAutoCleanupOperations
+	{
+		static bool HasElement(const FMassArchetypeHandle& ArchetypeHandle, TNotNull<UScriptStruct*> SparseElementType, bool bIncludeTransient)
+		{
+			return FMassArchetypeHelper::ArchetypeDataFromHandleChecked(ArchetypeHandle).HasFragmentType(SparseElementType, bIncludeTransient);
+		}
+
+		static bool HasElement(FMassEntityView EntityView, TNotNull<UScriptStruct*> SparseElementType)
+		{
+			return EntityView.HasFragment(SparseElementType);
+		}
+
+		TNotNull<UScriptStruct*> MakeSparseElement()
+		{
+			CreatedTypes.Add(MakeElement("SparseFragmentTestA", FMassSparseFragment::StaticStruct()));
+			return CreatedTypes.Last();
+		}
+
+		TNotNull<UScriptStruct*> MakeRegularElement()
+		{
+			CreatedTypes.Add(MakeElement("RegularFragmentTestA", FMassFragment::StaticStruct()));
+			return CreatedTypes.Last();
+		}
+	};
+
+	template<typename TTestOperations>
+	struct FSparseElementsBase : FEntityTestBase
+	{
+		TTestOperations TestOperations;
+		UScriptStruct* SparseElementType = nullptr;
+		FMassArchetypeEntityCollection InitialCollection;
+		TArray<FMassEntityHandle> EntitiesCreated;
+		FMassArchetypeHandle OriginalArchetypeHandle;
+		using FEntityTestBase::GetTestRunner;
+
+		virtual bool SetUp() override
+		{
+			FEntityTestBase::SetUp();
+			SparseElementType = TestOperations.MakeSparseElement();
+			OriginalArchetypeHandle = FEntityTestBase::IntsArchetype;
+			return true;
+		}
+	};
+
+#define USING_SUPER(SuperType) \
+		using Super = SuperType; \
+		using Super::EntityManager; \
+		using Super::OriginalArchetypeHandle; \
+		using Super::EntitiesCreated; \
+		using Super::InitialCollection; \
+		using Super::SparseElementType; \
+		using FEntityTestBase::GetTestRunner
+
+	template<typename TTestOperations>
+	struct FSparseElements_Add : FSparseElementsBase<TTestOperations>
+	{
+		USING_SUPER(FSparseElementsBase<TTestOperations>);
+
+		virtual bool InstantTest() override
+		{
+			constexpr int32 NumEntities = 2;
+			constexpr int32 ModifiedEntityHandleIndex = 1;
+
+			EntityManager->BatchCreateEntities(OriginalArchetypeHandle, NumEntities, EntitiesCreated);
+
+			const FMassEntityHandle ModifiedEntityHandle = EntitiesCreated[ModifiedEntityHandleIndex];
+			FStructView _ = EntityManager->AddSparseElementToEntity(ModifiedEntityHandle, SparseElementType);;
+
+			const FMassArchetypeHandle NewArchetypeHandle = EntityManager->GetArchetypeForEntity(ModifiedEntityHandle);
+			AITEST_TRUE(TEXT("The affected entity is still in the original archetype"), NewArchetypeHandle == OriginalArchetypeHandle);
+			AITEST_TRUE(TEXT("The archetype has the element")
+				, TTestOperations::HasElement(NewArchetypeHandle, SparseElementType, /*bIncludeTransient=*/true));
+
+			FMassEntityView ModifiedEntityView(*EntityManager, ModifiedEntityHandle);
+			AITEST_TRUE(TEXT("The affected entity has the sparse element"), TTestOperations::HasElement(ModifiedEntityView, SparseElementType));
+
+			FMassEntityView OriginalEntityView(*EntityManager, EntitiesCreated[(ModifiedEntityHandleIndex + 1) % NumEntities]);
+			AITEST_FALSE(TEXT("(NOT) The original entity has the sparse element"), TTestOperations::HasElement(OriginalEntityView, SparseElementType));
+
+			return true;
+		}
+	};
+	using FSparseElements_AddTag = FSparseElements_Add<FTagTestOperations>;
+	IMPLEMENT_AI_INSTANT_TEST(FSparseElements_AddTag, "System.Mass.Sparse.Tag.Add");
+	using FSparseElements_AddFragment = FSparseElements_Add<FFragmentTestOperations>;
+	IMPLEMENT_AI_INSTANT_TEST(FSparseElements_AddFragment, "System.Mass.Sparse.Fragment.Add");
+
+	template<typename TTestOperations>
+	struct FSparseElements_Remove : FSparseElementsBase<TTestOperations>
+	{
+		USING_SUPER(FSparseElementsBase<TTestOperations>);
+
+		virtual bool InstantTest() override
+		{
+			constexpr int32 NumEntities = 2;
+			
+			EntityManager->BatchCreateEntities(OriginalArchetypeHandle, NumEntities, EntitiesCreated);
+
+			for (FMassEntityHandle EntityHandle : EntitiesCreated)
+			{
+				FStructView _ = EntityManager->AddSparseElementToEntity(EntityHandle, SparseElementType);;
+			}
+
+			AITEST_TRUE(TEXT("First removal is expected to succeed"), EntityManager->RemoveSparseElementFromEntity(EntitiesCreated[0], SparseElementType));
+
+			const FMassArchetypeHandle NewArchetypeHandle = EntityManager->GetArchetypeForEntity(EntitiesCreated[0]);
+			AITEST_TRUE(TEXT("The affected entity is still in the original archetype"), NewArchetypeHandle == OriginalArchetypeHandle);
+			AITEST_TRUE(TEXT("The archetype knows some of its entities have the sparse element"), FMassArchetypeHelper::ArchetypeDataFromHandleChecked(NewArchetypeHandle).ContainsAnySparseData());
+
+			FMassEntityView ModifiedEntityView(*EntityManager, EntitiesCreated[0]);
+			AITEST_FALSE(TEXT("(NOT) The affected entity has the sparse element"), TTestOperations::HasElement(ModifiedEntityView, SparseElementType));
+
+			FMassEntityView OtherEntityView(*EntityManager, EntitiesCreated[1]);
+			AITEST_TRUE(TEXT("The other entity has the sparse element"), TTestOperations::HasElement(OtherEntityView, SparseElementType));
+
+			AITEST_TRUE(TEXT("Second removal is expected to succeed"),EntityManager->RemoveSparseElementFromEntity(EntitiesCreated[1], SparseElementType));
+			AITEST_FALSE(TEXT("(NOT) The other entity has the sparse element"), TTestOperations::HasElement(OtherEntityView, SparseElementType));
+
+			AITEST_TRUE(TEXT("Removing missing elements fails"),EntityManager->RemoveSparseElementFromEntity(EntitiesCreated[1], SparseElementType) == false);
+
+#if WITH_ARCHETYPE_UPDATING_BITSETS_ON_REMOVE
+			AITEST_TRUE(TEXT("The original archetype no longer has the element")
+				, TTestOperations::HasElement(OriginalArchetypeHandle, SparseElementType, /*bIncludeTransient=*/true) == false);
+#endif
+			return true;
+		}
+	};
+	using FSparseElements_RemoveTag = FSparseElements_Remove<FTagTestOperations>;
+	IMPLEMENT_AI_INSTANT_TEST(FSparseElements_RemoveTag, "System.Mass.Sparse.Tag.Remove");
+	using FSparseElements_RemoveFragment = FSparseElements_Remove<FFragmentTestOperations>;
+	IMPLEMENT_AI_INSTANT_TEST(FSparseElements_RemoveFragment, "System.Mass.Sparse.Fragment.Remove");
+
+	template<typename TTestOperations>
+	struct FSparseElements_MoveEntity : FSparseElementsBase<TTestOperations>
+	{
+		USING_SUPER(FSparseElementsBase<TTestOperations>);
+
+		virtual bool InstantTest() override
+		{
+			const FMassEntityHandle ModifiedEntityHandle = EntityManager->CreateEntity(OriginalArchetypeHandle);
+
+			FStructView _ = EntityManager->AddSparseElementToEntity(ModifiedEntityHandle, SparseElementType);;
+
+			EntityManager->AddFragmentToEntity(ModifiedEntityHandle, FTestFragment_Float::StaticStruct());
+#if WITH_ARCHETYPE_UPDATING_BITSETS_ON_REMOVE
+			AITEST_TRUE(TEXT("The original archetype no longer has the element")
+				, TTestOperations::HasElement(OriginalArchetypeHandle, SparseElementType, /*bIncludeTransient=*/true) == false);
+#endif
+			const FMassArchetypeHandle NewArchetypeHandle = EntityManager->GetArchetypeForEntity(ModifiedEntityHandle);
+			AITEST_TRUE(TEXT("The entity is expected to end up in the expected archetype"), NewArchetypeHandle == FEntityTestBase::FloatsIntsArchetype);
+
+			FMassEntityView ModifiedEntityView(*EntityManager, ModifiedEntityHandle);
+			AITEST_TRUE(TEXT("The affected entity has the sparse element"), TTestOperations::HasElement(ModifiedEntityView, SparseElementType));
+			
+			return true;
+		}
+	};
+	using FSparseElements_MoveEntityTag = FSparseElements_MoveEntity<FTagTestOperations>;
+	IMPLEMENT_AI_INSTANT_TEST(FSparseElements_MoveEntityTag, "System.Mass.Sparse.Tag.Move");
+	using FSparseElements_MoveEntityFragment = FSparseElements_MoveEntity<FFragmentTestOperations>;
+	IMPLEMENT_AI_INSTANT_TEST(FSparseElements_MoveEntityFragment, "System.Mass.Sparse.Fragment.Move");
+
+	// @todo add batch-add and batch-remove
+	template<typename TTestOperations>
+	struct FSparseElements_BatchMoveEntities : FSparseElementsBase<TTestOperations>
+	{
+		USING_SUPER(FSparseElementsBase<TTestOperations>);
+
+		virtual bool InstantTest() override
+		{
+			constexpr int32 NumEntities = 2;
+			
+			EntityManager->BatchCreateEntities(OriginalArchetypeHandle, NumEntities, EntitiesCreated);
+
+			for (FMassEntityHandle EntityHandle : EntitiesCreated)
+			{
+				FStructView _ = EntityManager->AddSparseElementToEntity(EntityHandle, SparseElementType);;
+			}
+
+			InitialCollection = FMassArchetypeEntityCollection(FEntityTestBase::IntsArchetype, EntitiesCreated, FMassArchetypeEntityCollection::NoDuplicates);
+			EntityManager->BatchChangeFragmentCompositionForEntities({InitialCollection}, FMassFragmentBitSet(*FTestFragment_Float::StaticStruct()), {});
+
+#if WITH_ARCHETYPE_UPDATING_BITSETS_ON_REMOVE
+			AITEST_TRUE(TEXT("The original archetype no longer has the element")
+				, TTestOperations::HasElement(OriginalArchetypeHandle, SparseElementType, /*bIncludeTransient=*/true) == false);
+#endif
+			FMassEntityHandle ModifiedEntityHandle = EntitiesCreated[0];
+			const FMassArchetypeHandle NewArchetypeHandle = EntityManager->GetArchetypeForEntity(ModifiedEntityHandle);
+			AITEST_TRUE(TEXT("The entity is expected to end up in the expected archetype"), NewArchetypeHandle == FEntityTestBase::FloatsIntsArchetype);
+
+			FMassEntityView ModifiedEntityView(*EntityManager, ModifiedEntityHandle);
+			AITEST_TRUE(TEXT("The affected entity has the sparse element"), TTestOperations::HasElement(ModifiedEntityView, SparseElementType));
+			
+			return true;
+		}
+	};
+	using FSparseElements_BatchMoveEntitiesTag = FSparseElements_BatchMoveEntities<FTagTestOperations>;
+	IMPLEMENT_AI_INSTANT_TEST(FSparseElements_BatchMoveEntitiesTag, "System.Mass.Sparse.Tag.BatchMove");
+	using FSparseElements_BatchMoveEntitiesFragment = FSparseElements_BatchMoveEntities<FFragmentTestOperations>;
+	IMPLEMENT_AI_INSTANT_TEST(FSparseElements_BatchMoveEntitiesFragment, "System.Mass.Sparse.Fragment.BatchMove");
+
+	template<typename TTestOperations>
+	struct FSparseElements_BatchAdd : FSparseElementsBase<TTestOperations>
+	{
+		USING_SUPER(FSparseElementsBase<TTestOperations>);
+
+		virtual bool InstantTest() override
+		{
+			constexpr int32 NumEntities = 2;
+
+			EntityManager->BatchCreateEntities(OriginalArchetypeHandle, NumEntities, EntitiesCreated);
+
+			InitialCollection = FMassArchetypeEntityCollection(FEntityTestBase::IntsArchetype, EntitiesCreated, FMassArchetypeEntityCollection::NoDuplicates);
+			EntityManager->BatchAddSparseElementToEntities({InitialCollection}, SparseElementType);
+
+#if WITH_ARCHETYPE_UPDATING_BITSETS_ON_REMOVE
+			AITEST_TRUE(TEXT("The original archetype no longer has the element")
+				, TTestOperations::HasElement(OriginalArchetypeHandle, SparseElementType, /*bIncludeTransient=*/true) == false);
+#endif
+
+			for (FMassEntityHandle ModifiedEntityHandle : EntitiesCreated)
+			{
+				const FMassArchetypeHandle NewArchetypeHandle = EntityManager->GetArchetypeForEntity(ModifiedEntityHandle);
+				AITEST_TRUE(TEXT("The entity is expected to end up in the expected archetype"), NewArchetypeHandle == OriginalArchetypeHandle);
+
+				FMassEntityView ModifiedEntityView(*EntityManager, ModifiedEntityHandle);
+				AITEST_TRUE(TEXT("The affected entities have the sparse element"), TTestOperations::HasElement(ModifiedEntityView, SparseElementType));
+			}
+			
+			return true;
+		}
+	};
+	using FSparseElements_BatchAddTag = FSparseElements_BatchAdd<FTagTestOperations>;
+	IMPLEMENT_AI_INSTANT_TEST(FSparseElements_BatchAddTag, "System.Mass.Sparse.Tag.BatchAdd");
+	using FSparseElements_BatchAddFragment = FSparseElements_BatchAdd<FFragmentTestOperations>;
+	IMPLEMENT_AI_INSTANT_TEST(FSparseElements_BatchAddFragment, "System.Mass.Sparse.Fragment.BatchAdd");
+
+	template<typename TTestOperations>
+	struct FSparseElements_BatchRemove : FSparseElements_BatchAdd<TTestOperations>
+	{
+		USING_SUPER(FSparseElements_BatchAdd<TTestOperations>);
+
+		virtual bool InstantTest() override
+		{
+			if (Super::InstantTest() == false)
+			{
+				return false;
+			}
+
+			EntityManager->BatchRemoveSparseElementFromEntities({InitialCollection}, SparseElementType);
+
+#if WITH_ARCHETYPE_UPDATING_BITSETS_ON_REMOVE
+			AITEST_TRUE(TEXT("The original archetype no longer has the element")
+				, TTestOperations::HasElement(OriginalArchetypeHandle, SparseElementType, /*bIncludeTransient=*/true) == false);
+#endif
+
+			for (FMassEntityHandle ModifiedEntityHandle : EntitiesCreated)
+			{
+				const FMassArchetypeHandle NewArchetypeHandle = EntityManager->GetArchetypeForEntity(ModifiedEntityHandle);
+				AITEST_TRUE(TEXT("The entity is expected to end up in the expected archetype"), NewArchetypeHandle == OriginalArchetypeHandle);
+
+				FMassEntityView ModifiedEntityView(*EntityManager, ModifiedEntityHandle);
+				AITEST_FALSE(TEXT("(NOT) The affected entities have the sparse element"), TTestOperations::HasElement(ModifiedEntityView, SparseElementType));
+			}
+			
+			return true;
+		}
+	};
+	using FSparseElements_BatchRemoveTag = FSparseElements_BatchRemove<FTagTestOperations>;
+	IMPLEMENT_AI_INSTANT_TEST(FSparseElements_BatchRemoveTag, "System.Mass.Sparse.Tag.BatchRemove");
+	using FSparseElements_BatchRemoveFragment = FSparseElements_BatchRemove<FFragmentTestOperations>;
+	IMPLEMENT_AI_INSTANT_TEST(FSparseElements_BatchRemoveFragment, "System.Mass.Sparse.Fragment.BatchRemove");
+
+	template<typename TTestOperations>
+	struct FSparseElements_Query : FSparseElementsBase<TTestOperations>
+	{
+		USING_SUPER(FSparseElementsBase<TTestOperations>);
+
+		virtual bool InstantTest() override
+		{
+			constexpr int32 NumEntities = 5;
+			constexpr int32 NumEntitiesAffected = 3;
+			EntityManager->BatchCreateEntities(OriginalArchetypeHandle, NumEntities, EntitiesCreated);
+
+			TConstArrayView<FMassEntityHandle> ModifiedEntities = MakeArrayView(&EntitiesCreated[1], NumEntitiesAffected);
+			InitialCollection = FMassArchetypeEntityCollection(FEntityTestBase::IntsArchetype, ModifiedEntities, FMassArchetypeEntityCollection::NoDuplicates);
+			EntityManager->BatchAddSparseElementToEntities({InitialCollection}, SparseElementType);
+
+			// 1. Regular + Sparse
+			{
+				TArray<FMassEntityHandle> VerifyEntities1;
+
+				FMassEntityQuery Query(EntityManager);
+				Query.AddRequirement<FTestFragment_Int>(EMassFragmentAccess::ReadOnly);
+				Query.AddSparseRequirement(SparseElementType);
+
+				FMassExecutionContext ExecutionContext = EntityManager->CreateExecutionContext(0);
+				FAutomationTestBase& LocalTestRunner = GetTestRunner();
+				Query.ForEachEntityChunk(ExecutionContext, [SparseElementType=SparseElementType, &VerifyEntities1, &LocalTestRunner](FMassExecutionContext& Context)
+				{
+					// option 1: the archetype has the sparse element, and we need to check individual entities
+					for (auto EntityIt = Context.CreateEntityIterator(); EntityIt; ++EntityIt)
+					{
+						if (const bool bHasSparseElement = Context.HasSparseElement(SparseElementType, EntityIt))
+						{
+							if (UE::Mass::IsA<FMassFragment>(SparseElementType))
+							{
+								FConstStructView ElementView = Context.GetSparseElement(SparseElementType, EntityIt);
+								AITEST_TRUE_EX(LocalTestRunner, TEXT("The filtered entity has the sparse element, const view"), ElementView.IsValid());
+
+								FStructView ElementMutableView = Context.GetMutableSparseElement(SparseElementType, EntityIt);
+								AITEST_TRUE_EX(LocalTestRunner, TEXT("The filtered entity has the sparse element, mutable view"), ElementMutableView.IsValid());
+
+								AITEST_TRUE_EX(LocalTestRunner, TEXT("Both views point at the same element"), ElementView == ElementMutableView);
+							}
+
+							VerifyEntities1.Add(Context.GetEntities()[EntityIt]);
+						}
+					}
+					return true;
+				});
+
+				AITEST_EQUAL("Number of expected entities, option 1", VerifyEntities1.Num(), NumEntitiesAffected);
+
+				TArray<FMassEntityHandle> VerifyEntities2;
+				Query.ForEachEntityChunk(ExecutionContext, [SparseElementType=SparseElementType, &VerifyEntities2, &LocalTestRunner](FMassExecutionContext& Context)
+				{
+					// option 2: dedicated sparse element iterator that skips entities that don't have the elements
+					// CreateSparseEntityIterator:
+					// * derives from FMassExecutionContext::FEntityIterator and overrides operator++ to skip elements not matching the sparse elements composition
+					// * compares against sparse-elements requirements expressed by the query (as stored in the Context).
+					for (FMassSparseEntityIterator EntitySparseIt = Context.CreateSparseEntityIterator(); EntitySparseIt; ++EntitySparseIt)
+					{
+						// all the following would be expected to succeed.
+						const bool bHasSparseElement = Context.HasSparseElement(SparseElementType, EntitySparseIt);
+						AITEST_TRUE_EX(LocalTestRunner, TEXT("The filtered entity has the sparse element"), bHasSparseElement);
+
+						if (UE::Mass::IsA<FMassFragment>(SparseElementType))
+						{
+							FConstStructView ElementView = Context.GetSparseElement(SparseElementType, EntitySparseIt);
+							AITEST_TRUE_EX(LocalTestRunner, TEXT("The filtered entity has the sparse element, const view"), ElementView.IsValid());
+
+							FStructView ElementMutableView = Context.GetMutableSparseElement(SparseElementType, EntitySparseIt);
+							AITEST_TRUE_EX(LocalTestRunner, TEXT("The filtered entity has the sparse element, mutable view"), ElementMutableView.IsValid());
+
+							AITEST_TRUE_EX(LocalTestRunner, TEXT("Both views point at the same element"), ElementView == ElementMutableView);
+						}
+
+						VerifyEntities2.Add(Context.GetEntities()[EntitySparseIt]);
+					}
+
+					return true;
+				});
+
+				AITEST_EQUAL("Number of expected entities, option 2", VerifyEntities2.Num(), NumEntitiesAffected);
+			}
+
+			// 2. No-sparse
+			{
+				FMassEntityQuery Query(EntityManager);
+				Query.AddRequirement<FTestFragment_Int>(EMassFragmentAccess::ReadOnly/*, EMassFragmentAccess::ReadOnly*/);
+				Query.AddSparseRequirement(SparseElementType, EMassFragmentPresence::None);
+
+				FMassExecutionContext ExecutionContext = EntityManager->CreateExecutionContext(0);
+				FAutomationTestBase& LocalTestRunner = GetTestRunner();
+				
+				TArray<FMassEntityHandle> VerifyEntities;
+				Query.ForEachEntityChunk(ExecutionContext, [SparseElementType=SparseElementType, &VerifyEntities, &LocalTestRunner](FMassExecutionContext& Context)
+				{
+					for (FMassSparseEntityIterator EntitySparseIt = Context.CreateSparseEntityIterator(); EntitySparseIt; ++EntitySparseIt)
+					{
+						const bool bHasSparseElement = Context.HasSparseElement(SparseElementType, EntitySparseIt);
+						AITEST_FALSE_EX(LocalTestRunner, TEXT("(NOT) The filtered entity has the sparse element"), bHasSparseElement);
+						VerifyEntities.Add(Context.GetEntities()[EntitySparseIt]);
+					}
+
+					return true;
+				});
+
+				AITEST_EQUAL("Number of expected entities, no-sparse", VerifyEntities.Num(), (NumEntities - NumEntitiesAffected));
+			}
+
+			return true;
+		}
+	};
+	using FSparseElements_QueryTag = FSparseElements_Query<FTagTestOperations>;
+	IMPLEMENT_AI_INSTANT_TEST(FSparseElements_QueryTag, "System.Mass.Sparse.Query.Tag");
+	using FSparseElements_QueryFragment = FSparseElements_Query<FFragmentTestOperations>;
+	IMPLEMENT_AI_INSTANT_TEST(FSparseElements_QueryFragment, "System.Mass.Sparse.Query.Fragment");
+
+
+	struct FSparseElements_QueryFragmentModification : FEntityTestBase
+	{
+		virtual bool InstantTest() override
+		{
+			constexpr int32 NumEntities = 5;
+			constexpr int32 NumEntitiesAffected = 3;
+
+			TArray<FMassEntityHandle> EntitiesCreated;
+			EntityManager->BatchCreateEntities(IntsArchetype, NumEntities, EntitiesCreated);
+
+			TConstArrayView<FMassEntityHandle> ModifiedEntities = MakeArrayView(&EntitiesCreated[1], NumEntitiesAffected);
+			FMassArchetypeEntityCollection InitialCollection = FMassArchetypeEntityCollection(FEntityTestBase::IntsArchetype, ModifiedEntities, FMassArchetypeEntityCollection::NoDuplicates);
+			EntityManager->BatchAddSparseElementToEntities({InitialCollection}, FTestFragment_SparseInt::StaticStruct());
+
+			FMassEntityQuery Query(EntityManager);
+			Query.AddRequirement<FTestFragment_Int>(EMassFragmentAccess::ReadOnly);
+			Query.AddSparseRequirement<FTestFragment_SparseInt>();
+
+			FMassExecutionContext ExecutionContext = EntityManager->CreateExecutionContext(0);
+			FAutomationTestBase& LocalTestRunner = GetTestRunner();
+
+			int32 Counter = 0;
+			TArray<FMassEntityHandle> VerifyEntities;
+			Query.ForEachEntityChunk(ExecutionContext, [&Counter, &VerifyEntities, &LocalTestRunner](FMassExecutionContext& Context)
+			{
+				for (FMassSparseEntityIterator EntitySparseIt = Context.CreateSparseEntityIterator(); EntitySparseIt; ++EntitySparseIt)
+				{
+					FTestFragment_SparseInt* FragmentInstance = Context.GetMutableSparseElement<FTestFragment_SparseInt>(EntitySparseIt);
+					AITEST_NOT_NULL_EX(LocalTestRunner, TEXT("The filtered entity has the sparse element"), FragmentInstance);
+
+					FragmentInstance->Value = Counter++;
+
+					VerifyEntities.Add(Context.GetEntities()[EntitySparseIt]);
+				}
+
+				return true;
+			});
+
+			for (int32 EntityIndex = 0; EntityIndex < VerifyEntities.Num(); ++EntityIndex)
+			{
+				FMassEntityView EntityView(*EntityManager, VerifyEntities[EntityIndex]);
+				FTestFragment_SparseInt* FragmentPtr = EntityView.GetSparseFragmentDataPtr<FTestFragment_SparseInt>();
+				AITEST_NOT_NULL(TEXT("Collected entity's sparse fragment"), FragmentPtr);
+				AITEST_EQUAL(TEXT("Sparse fragment's expected value"), FragmentPtr->Value, EntityIndex);
+			}
+
+			return true;
+		}
+	};
+	IMPLEMENT_AI_INSTANT_TEST(FSparseElements_QueryFragmentModification, "System.Mass.Sparse.Query.ModifyFragment");
+
+	struct FSparseElements_QuerySkipChunks : FEntityTestBase
+	{
+		virtual bool InstantTest() override
+		{
+			const int32 EntitiesPerChunk = EntityManager->DebugGetArchetypeEntitiesCountPerChunk(IntsArchetype);
+			constexpr int32 NumTotalChunks = 5;
+			const int32 NumEntities = NumTotalChunks * EntitiesPerChunk;
+			constexpr int32 NumAffectedChunks = 2;
+			const int32 NumEntitiesAffected = NumAffectedChunks * EntitiesPerChunk;
+
+			TArray<FMassEntityHandle> EntitiesCreated;
+			EntityManager->BatchCreateEntities(IntsArchetype, NumEntities, EntitiesCreated);
+
+			// ModifiedEntities should contain all the entities in chunks 1 & 2
+			TConstArrayView<FMassEntityHandle> ModifiedEntities = MakeArrayView(&EntitiesCreated[EntitiesPerChunk], NumEntitiesAffected);
+			FMassArchetypeEntityCollection InitialCollection = FMassArchetypeEntityCollection(FEntityTestBase::IntsArchetype, ModifiedEntities, FMassArchetypeEntityCollection::NoDuplicates);
+			EntityManager->BatchAddSparseElementToEntities({InitialCollection}, FTestFragment_SparseInt::StaticStruct());
+
+			FMassExecutionContext ExecutionContext = EntityManager->CreateExecutionContext(0);
+			FMassEntityQuery Query(EntityManager);
+			Query.AddRequirement<FTestFragment_Int>(EMassFragmentAccess::ReadOnly);
+			Query.AddSparseRequirement<FTestFragment_SparseInt>();
+
+			int32 Counter = 0;
+			Query.ForEachEntityChunk(ExecutionContext, [&Counter](FMassExecutionContext& Context)
+			{
+				++Counter;
+			});
+			AITEST_EQUAL(TEXT("Expected number of chunks processed"), Counter, NumAffectedChunks);
+
+			return true;
+		}
+	};
+	IMPLEMENT_AI_INSTANT_TEST(FSparseElements_QuerySkipChunks, "System.Mass.Sparse.Query.SkipChunks");
+
+	struct FSparseElements_Storage_Individual : FEntityTestBase
+	{
+		virtual bool InstantTest() override
+		{
+			FSparseElementsStorage Storage;
+
+			const FMassEntityHandle EntityHandle = EntityManager->CreateEntity(IntsArchetype);
+
+			FConstStructView NoView = Storage.GetElementDataForEntity(EntityHandle, FTestFragment_SparseInt::StaticStruct());
+			AITEST_TRUE(TEXT("Initially entity doesn't have the sparse element"), NoView.IsValid() == false);
+
+			FStructView ElementView = Storage.AddElementToEntity(EntityHandle, FTestFragment_SparseInt::StaticStruct());
+			AITEST_TRUE(TEXT("View of added element is valid"), ElementView.IsValid());
+
+			FStructView ExistingView = Storage.GetMutableElementDataForEntity(EntityHandle, FTestFragment_SparseInt::StaticStruct());
+			AITEST_TRUE(TEXT("Existing fragment view is valid"), ExistingView.IsValid());
+			AITEST_EQUAL(TEXT("Both views point to the same fragment"), ElementView, ExistingView);
+
+			const bool bRemoved = Storage.RemoveElementFromEntity(EntityHandle, FTestFragment_SparseInt::StaticStruct());
+			AITEST_TRUE(TEXT("Removing the added element succeeds"), bRemoved);
+
+			const bool bSecondRemove = Storage.RemoveElementFromEntity(EntityHandle, FTestFragment_SparseInt::StaticStruct());
+			AITEST_TRUE(TEXT("Removing the added element for the second time fails"), bSecondRemove == false);
+
+			FConstStructView RemovedView = Storage.GetElementDataForEntity(EntityHandle, FTestFragment_SparseInt::StaticStruct());
+			AITEST_TRUE(TEXT("Fetching the removed element results in a N invalid view"), RemovedView.IsValid() == false);
+
+			return true;
+		}
+	};
+	IMPLEMENT_AI_INSTANT_TEST(FSparseElements_Storage_Individual, "System.Mass.Sparse.Storage.IndividualBasic");
+
+	struct FSparseElements_Storage_Iterator : FEntityTestBase
+	{
+		virtual bool InstantTest() override
+		{
+			constexpr int32 NumEntities = 100;
+			int32 NumWithSparseElement = 0;
+
+			FRandomStream RandStream(0);
+
+			FSparseElementsStorage Storage;
+
+			TArray<FMassEntityHandle> EntitiesCreated;
+			EntityManager->BatchCreateEntities(IntsArchetype, NumEntities, EntitiesCreated);
+
+			// add sparse elements
+			TArray<FMassEntityHandle> ModifiedEntities;
+			for (int32 EntityIndex = 0; EntityIndex < EntitiesCreated.Num(); ++EntityIndex)
+			{
+				if (RandStream.FRand() < 0.3f)
+				{
+					Storage.AddElementToEntity<FTestFragment_SparseInt>(EntitiesCreated[EntityIndex]);
+					ModifiedEntities.Add(EntitiesCreated[EntityIndex]);
+				}
+			}
+
+			int32 IteratedCount = 0;
+			TSet<int32> IteratedEntityIndices;
+
+			for (UE::Mass::FSparseElementIterator It = Storage.CreateElementIterator(FTestFragment_SparseInt::StaticStruct()); It; ++It)
+			{
+				++IteratedCount;
+
+				FStructView ElementView = It.GetElementView();
+				AITEST_TRUE(TEXT("Iterated element view is valid"), ElementView.IsValid());
+
+				const int32 EntityIndex = It.GetEntityIndex();
+				IteratedEntityIndices.Add(EntityIndex);
+
+				// perform modification, we'll test it later on
+				FTestFragment_SparseInt* Fragment = ElementView.GetPtr<FTestFragment_SparseInt>();
+				Fragment->Value = EntityIndex;
+			}
+
+			AITEST_EQUAL(TEXT("Iterated over all modified entities"), IteratedCount, ModifiedEntities.Num());
+
+			for (const FMassEntityHandle& Entity : ModifiedEntities)
+			{
+				AITEST_TRUE(TEXT("Iterated over modified entity"), IteratedEntityIndices.Contains(Entity.Index));
+			}
+
+			// test the values written
+			for (UE::Mass::FSparseElementIterator It = Storage.CreateElementIterator(FTestFragment_SparseInt::StaticStruct()); It; ++It)
+			{
+				FConstStructView ElementView = It.GetElementView();
+				AITEST_TRUE(TEXT("Const view is valid"), ElementView.IsValid());
+
+				const int32 EntityIndex = It.GetEntityIndex();
+				const FTestFragment_SparseInt* Fragment = ElementView.GetPtr<FTestFragment_SparseInt>();
+				AITEST_EQUAL(TEXT("Stored value matches expectations"), Fragment->Value, EntityIndex);
+			}
+
+			return true;
+		}
+	};
+	IMPLEMENT_AI_INSTANT_TEST(FSparseElements_Storage_Iterator, "System.Mass.Sparse.Storage.Iterator");
+
+#undef USING_SUPER
+#endif // WITH_MASSENTITY_DEBUG
+} // namespace
+
+#undef WITH_ARCHETYPE_UPDATING_BITSETS_ON_REMOVE
+
+UE_ENABLE_OPTIMIZATION_SHIP
+
+#undef LOCTEXT_NAMESPACE
diff --git a/Engine/Source/Developer/MassEntityTestSuite/Public/MassEntityTestTypes.h b/Engine/Source/Developer/MassEntityTestSuite/Public/MassEntityTestTypes.h
--- a/Engine/Source/Developer/MassEntityTestSuite/Public/MassEntityTestTypes.h
+++ b/Engine/Source/Developer/MassEntityTestSuite/Public/MassEntityTestTypes.h
@@ -257,6 +257,44 @@ struct FTestTag_D : public FMassTag
 	GENERATED_BODY()
 };
 
+USTRUCT()
+struct FTestFragment_SparseInt : public FMassSparseFragment
+{
+	GENERATED_BODY()
+
+	FTestFragment_SparseInt(const int32 InValue = 0)
+		: Value(InValue)
+	{
+	}
+
+	int32 Value = 0;
+};
+
+USTRUCT()
+struct FTestFragment_SparseFloat : public FMassSparseFragment
+{
+	GENERATED_BODY()
+
+	FTestFragment_SparseFloat(const float InValue = 0.f)
+		: Value(InValue)
+	{
+	}
+
+	float Value = 0.f;
+};
+
+USTRUCT()
+struct FTestTag_SparseA : public FMassSparseTag
+{
+	GENERATED_BODY()
+};
+
+USTRUCT()
+struct FTestTag_SparseB : public FMassSparseTag
+{
+	GENERATED_BODY()
+};
+
 
 UCLASS(MinimalAPI)
 class UMassTestProcessorBase : public UMassProcessor
diff --git a/Engine/Source/Runtime/MassEntity/Internal/MassArchetypeData.h b/Engine/Source/Runtime/MassEntity/Internal/MassArchetypeData.h
--- a/Engine/Source/Runtime/MassEntity/Internal/MassArchetypeData.h
+++ b/Engine/Source/Runtime/MassEntity/Internal/MassArchetypeData.h
@@ -8,6 +8,8 @@
 #if UE_ENABLE_INCLUDE_ORDER_DEPRECATED_IN_5_6
 #include "MassEntityManager.h"
 #endif // UE_ENABLE_INCLUDE_ORDER_DEPRECATED_IN_5_6
+#include "MassRequirements.h"
+#include "MassSparseElementsStorage.h"
 #include "HAL/LowLevelMemTracker.h"
 
 struct FMassEntityQuery;
@@ -19,10 +21,88 @@ struct FMassFragmentRequirements;
 
 namespace UE::Mass
 {
+	struct FSparseElementsStorage;
+	struct FChunkSparseElements;
+
 	uint32 SanitizeChunkMemorySize(const uint32 InChunkMemorySize, const bool bLogMismatch = true);
+
+	struct FChunkSparseElements
+	{
+		FChunkSparseElements() = default;
+		FChunkSparseElements(const uint32 NumEntitiesPerChunk)
+		{
+			Init(NumEntitiesPerChunk);
+		}
+
+		void Init(const uint32 NumEntitiesPerChunk)
+		{
+			PerEntityElements.AddDefaulted(NumEntitiesPerChunk);
+		}
+
+		void Add(const FMassRawEntityInChunkData& EntityInChunkHandle, TNotNull<const UScriptStruct*> ElementType);
+		bool Remove(const FMassRawEntityInChunkData& EntityInChunkHandle, TNotNull<const UScriptStruct*> ElementType);
+
+		void BatchAdd(TNotNull<const UScriptStruct*> ElementType, int32 SubchunkStart, int32 RangeLength);
+		void BatchRemove(TNotNull<const UScriptStruct*> ElementType, int32 SubchunkStart, int32 RangeLength);
+
+		bool IsEmpty() const
+		{
+			return PerEntityElements.IsEmpty();
+		}
+
+		bool Contains(const FMassRawEntityInChunkData& EntityInChunkHandle, TNotNull<const UScriptStruct*> ElementType) const;
+		/** @return whether anything in the chunk potentially has the element */
+		bool Contains(TNotNull<const UScriptStruct*> ElementType) const
+		{
+			return SparseElementsPresent.Contains(*ElementType);
+		}
+
+		bool HasElementsInRange(int32 IndexWithinChunk, int32 Num) const
+		{
+			for (int32 Index = 0; Index < Num; ++Index)
+			{
+				if (PerEntityElements.IsValidIndex(IndexWithinChunk + Index) 
+					&& PerEntityElements[IndexWithinChunk + Index].IsEmpty() == false)
+				{
+					return true;
+				}
+			}
+			return false;
+		}
+
+		void MoveElementsFrom(FChunkSparseElements& OriginalSparseElements, int32 OriginalIndexWithinChunk, int32 NewIndexWithinChunk, int32 Num);
+
+		const FMassElementBitSet& GetSparseElementsPresent() const
+		{
+			return SparseElementsPresent;
+		}
+
+		FMassElementBitSet& GetElementsForEntity(const FMassRawEntityInChunkData& EntityInChunkHandle)
+		{
+			check(PerEntityElements.IsEmpty() == false);
+			return PerEntityElements[EntityInChunkHandle.IndexWithinChunk];
+		}
+
+		const FMassElementBitSet& GetElementsForEntity(const FMassRawEntityInChunkData& EntityInChunkHandle) const
+		{
+			check(PerEntityElements.IsEmpty() == false);
+			return PerEntityElements[EntityInChunkHandle.IndexWithinChunk];
+		}
+
+		FMassElementBitSet& GetElementsForEntityUnsafe(const int32 Index)
+		{
+			check(PerEntityElements.IsValidIndex(Index));
+			return PerEntityElements[Index];
+		}
+
+
+	private:
+		FMassElementBitSet SparseElementsPresent;
+		TArray<FMassElementBitSet> PerEntityElements;
+	};
 }
 
-// This is one chunk within an archetype
+/** This type represents a single chunk within an archetype */
 struct FMassArchetypeChunk
 {
 private:
@@ -32,7 +112,8 @@ private:
 	int32 SerialModificationNumber = 0;
 	TArray<FInstancedStruct> ChunkFragmentData;
 	FMassArchetypeSharedFragmentValues SharedFragmentValues;
-
+	UE::Mass::FChunkSparseElements SparseElements;
+	
 public:
 	explicit FMassArchetypeChunk(const SIZE_T InAllocSize, TConstArrayView<FInstancedStruct> InChunkFragmentTemplates, const FMassArchetypeSharedFragmentValues& InSharedFragmentValues)
 		: AllocSize(InAllocSize)
@@ -81,6 +162,42 @@ public:
 		return NumInstances;
 	}
 
+	bool HasSparseElements() const
+	{
+		return SparseElements.IsEmpty() == false;
+	}
+
+	UE::Mass::FChunkSparseElements& GetSparseElements(const uint32 NumEntitiesPerChunk)
+	{
+		if (SparseElements.IsEmpty())
+		{
+			SparseElements.Init(NumEntitiesPerChunk);
+		}
+		return SparseElements;
+	}
+
+	UE::Mass::FChunkSparseElements& GetSparseElementsUnsafe()
+	{
+		check(SparseElements.IsEmpty() == false);
+		return SparseElements;
+	}
+
+	const UE::Mass::FChunkSparseElements& GetSparseElementsUnsafe() const
+	{
+		check(SparseElements.IsEmpty() == false);
+		return SparseElements;
+	}
+
+	bool DoesMatchComposition(const FMassElementBitSet& RequiredAllSparseElements, const FMassElementBitSet& RequiredAnySparseElements) const
+	{
+		if (SparseElements.IsEmpty() == false)
+		{
+			return (RequiredAllSparseElements.IsEmpty() || SparseElements.GetSparseElementsPresent().HasAll(RequiredAllSparseElements))
+				&& (RequiredAnySparseElements.IsEmpty() || SparseElements.GetSparseElementsPresent().HasAny(RequiredAnySparseElements));
+		}
+		return RequiredAllSparseElements.IsEmpty() && RequiredAnySparseElements.IsEmpty();
+	}
+
 	void AddMultipleInstances(uint32 Count)
 	{
 		NumInstances += Count;
@@ -173,6 +290,7 @@ struct FMassArchetypeData
 private:
 	// One-stop-shop variable describing composition of entities hosted by the archetype
 	FMassElementBitSet CompositionBitSet;
+	FMassElementBitSet SparseElementsBitSet;
 
 	// Pre-created default chunk fragment templates
 	TArray<FInstancedStruct> ChunkFragmentsTemplate;
@@ -209,6 +327,12 @@ private:
 	/** Defaults to UMassEntitySettings.ChunkMemorySize. In near future will support being set via constructor. */
 	const uint32 ChunkMemorySize = 0;
 
+	/**
+	 * Points to the EntityManager-owned SparseElementsStorage. This is why it's safe to use a raw pointer here,
+	 * because that's the manager that the sole owner and creator of FMassArchetypeData instances
+	 */
+	UE::Mass::FSparseElementsStorage* SparseElementsStorage = nullptr;
+
 #if WITH_MASSENTITY_DEBUG
 	/** Arrays of names the archetype is referred as. */
 	TArray<FName> DebugNames;
@@ -250,27 +374,92 @@ public:
 
 	/** Method to iterate on all the fragment types */
 	void ForEachFragmentType(TFunction< void(const UScriptStruct* /*FragmentType*/)> Function) const;
-	bool HasFragmentType(const UScriptStruct* FragmentType) const;
-	bool HasTagType(const UScriptStruct* TagType) const;
+
+	/**
+	 * checks whether a given fragment type is part of this archetype's entities' composition
+	 * @param bIncludeSparse if true and FragmentType is a sparse element, the function will
+	 *	check if it hosts any entities owning such a fragment
+	 */
+	bool HasFragmentType(const UScriptStruct* FragmentType, const bool bIncludeSparse = false) const
+	{
+		return (UE::Mass::IsA<FMassFragment>(FragmentType) 
+			&& (CompositionBitSet.Contains(*FragmentType)
+				|| (bIncludeSparse && UE::Mass::IsSparse(FragmentType) && SparseElementsBitSet.Contains(*FragmentType))));
+	}
+
+	/**
+	 * checks whether a given tag type is part of this archetype's entities' composition
+	 * @param bIncludeSparse if true and TagType is a sparse element, the function will
+	 *	check if it hosts any entities owning such a tag
+	 */
+	bool HasTagType(const UScriptStruct* TagType, const bool bIncludeSparse = false) const
+	{
+		return (UE::Mass::IsA<FMassTag>(TagType)
+			&& (CompositionBitSet.Contains(*TagType)
+				|| (bIncludeSparse && UE::Mass::IsSparse(TagType) && SparseElementsBitSet.Contains(*TagType))));
+	}
+
+	/**
+	 * checks whether a given element type is part of this archetype's entities' composition
+	 * @param bIncludeSparse if true and ElementType is a sparse element, the function will
+	 *	check if it hosts any entities owning such an element
+	 */
+	bool HasElement(const UScriptStruct* ElementType, const bool bIncludeSparse = false) const
+	{
+		return (CompositionBitSet.Contains(*ElementType)
+				|| (bIncludeSparse && UE::Mass::IsSparse(ElementType) && SparseElementsBitSet.Contains(*ElementType)));
+	}
 
 	bool IsEquivalent(const FMassElementBitSet& OtherCompositionBitSet, const UE::Mass::FArchetypeGroups& OtherGroups) const;
 
-	void Initialize(const FMassEntityManager& EntityManager, const FMassElementBitSet& InCompositionBitSet, const uint32 ArchetypeDataVersion);
+	void Initialize(const FMassEntityManager& EntityManager, const FMassElementBitSet& InCompositionBitSet
+		, const uint32 ArchetypeDataVersion, UE::Mass::FSparseElementsStorage& InSparseElementsStorage);
 
 	/** 
 	 * A special way of initializing an archetype resulting in a copy of BaseArchetype's setup with OverrideTags
 	 * replacing original tags of BaseArchetype
 	 */
 	void InitializeWithSimilar(const FMassEntityManager& EntityManager, const FMassArchetypeData& BaseArchetype
-		, FMassElementBitSet&& InCompositionBitSet, const UE::Mass::FArchetypeGroups& InGroups, const uint32 ArchetypeDataVersion);
+		, FMassElementBitSet&& InCompositionBitSet, const UE::Mass::FArchetypeGroups& InGroups, const uint32 ArchetypeDataVersion
+		, UE::Mass::FSparseElementsStorage& InSparseElementsStorage);
 
 	void AddEntity(FMassEntityHandle Entity, const FMassArchetypeSharedFragmentValues& InSharedFragmentValues);
 	void RemoveEntity(FMassEntityHandle Entity);
 
-	bool HasFragmentDataForEntity(const UScriptStruct* FragmentType, int32 EntityIndex) const;
+	/** @return if ElementType is not a tag the function will return FInstancedStruct ptr pointing at the appropriate data */
+	FStructView AddSparseElementToEntity(FMassEntityHandle EntityHandle, TNotNull<const UScriptStruct*> ElementType);
+	/** Works with all sparse element types (fragments and tags). */
+	bool RemoveSparseElementFromEntity(FMassEntityHandle EntityHandle, TNotNull<const UScriptStruct*> ElementType);
+	/** Expects non-tag sparse element type. Will assert if the assumption is broken. */
+	FStructView GetMutableSparseElementDataForEntity(TNotNull<const UScriptStruct*> ElementType, const FMassEntityHandle EntityHandle) const;
+	/** Expects non-tag sparse element type. Will assert if the assumption is broken. */
+	FConstStructView GetSparseElementDataForEntity(TNotNull<const UScriptStruct*> ElementType, const FMassEntityHandle EntityHandle) const;
+	/** @return a bitset indicating all sparse elements (both tags and fragments) owned by the given entity */
+	FMassElementBitSet GetSparseElementsBitSetForEntity(const FMassEntityHandle EntityHandle) const;
+
 	void* GetFragmentDataForEntityChecked(TNotNull<const UScriptStruct*> FragmentType, int32 EntityIndex) const;
 	void* GetFragmentDataForEntity(TNotNull<const UScriptStruct*> FragmentType, int32 EntityIndex) const;
 
+	bool HasSparseElementForEntity(TNotNull<const UScriptStruct*> ElementType, const FMassEntityHandle EntityHandle) const;
+	bool HasSparseElementForEntity(TNotNull<const UScriptStruct*> ElementType, const FMassEntityInChunkDataHandle& EntityInChunkHandle) const;
+	bool DoesChunkContainSparseElement(TNotNull<const UScriptStruct*> ElementType, const int32 ChunkIndex) const
+	{
+		checkf(Chunks.IsValidIndex(ChunkIndex), TEXT("Invalid chunk index (%d out of %d)"), ChunkIndex, Chunks.Num());
+
+		if (SparseElementsBitSet.Contains(*ElementType) && Chunks[ChunkIndex].HasSparseElements())
+		{
+			const UE::Mass::FChunkSparseElements& SparseElements = Chunks[ChunkIndex].GetSparseElementsUnsafe();
+			return (SparseElements.IsEmpty() == false) && SparseElements.Contains(ElementType);
+		}
+		return false;
+	}
+	bool DoesEntityMatchSparseComposition(const FMassEntityInChunkDataHandle& EntityInChunkHandle
+		, const FMassElementBitSet& RequiredAllSparseElements
+		, const FMassElementBitSet& RequiredAnySparseElements
+		, const FMassElementBitSet& RequiredNoneSparseElements) const;
+	bool ContainsAnySparseData() const;
+	bool ContainsSparseElement(TNotNull<const UScriptStruct*> ElementType) const;
+
 	FORCEINLINE const int32* GetInternalIndexForEntity(const int32 EntityIndex) const { return EntityMap.Find(EntityIndex); }
 	FORCEINLINE int32 GetInternalIndexForEntityChecked(const int32 EntityIndex) const { return EntityMap.FindChecked(EntityIndex); }
 	int32 GetNumEntitiesPerChunk() const { return NumEntitiesPerChunk; }
@@ -435,6 +624,7 @@ public:
 	{
 		const int32 AbsoluteIndex = EntityMap.FindChecked(EntityIndex);
 		const int32 ChunkIndex = AbsoluteIndex / NumEntitiesPerChunk;
+		checkf(Chunks.IsValidIndex(ChunkIndex), TEXT("Provided Entity Index %d is not hosted by this archetype"), EntityIndex);
 		const FMassArchetypeChunk& Chunk = Chunks[ChunkIndex];
 
 		return FMassEntityInChunkDataHandle(Chunk.GetRawMemory(), AbsoluteIndex - (NumEntitiesPerChunk * ChunkIndex)
@@ -463,6 +653,11 @@ public:
 		, const FMassConstSharedFragmentBitSet* ConstSharedFragmentToRemoveBitSet = nullptr);
 	void BatchSetFragmentValues(TConstArrayView<FMassArchetypeEntityCollection::FArchetypeEntityRange> EntityCollection, const FMassGenericPayloadViewSlice& Payload);
 
+	/** Works with all sparse element types (fragments and tags). */
+	void BatchAddSparseElementToEntities(const FMassArchetypeEntityCollection& EntityCollection, TNotNull<const UScriptStruct*> ElementType);
+	/** Works with all sparse element types (fragments and tags). */
+	void BatchRemoveSparseElementFromEntities(const FMassArchetypeEntityCollection& EntityCollection, TNotNull<const UScriptStruct*> ElementType);
+
 protected:
 	FMassArchetypeEntityCollection::FArchetypeEntityRange PrepareNextEntitiesSpanInternal(TConstArrayView<FMassEntityHandle> Entities, const FMassArchetypeSharedFragmentValues& InSharedFragmentValues, const int32 StartingChunk = 0);
 	void BatchRemoveEntitiesInternal(const int32 ChunkIndex, const int32 StartIndexWithinChunk, const int32 NumberToRemove);
@@ -472,7 +667,9 @@ protected:
 		uint8* RawChunkMemory;
 		int32 IndexWithinChunk;
 	};
-	void MoveFragmentsToAnotherArchetypeInternal(FMassArchetypeData& TargetArchetype, FTransientChunkLocation Target, const FTransientChunkLocation Source, const int32 ElementsNum);
+
+	void MoveFragmentsToAnotherArchetypeInternal(FMassArchetypeData& TargetArchetype, FMassArchetypeChunk& TargetChunk, int32 NewIndexWithinChunk
+		, FMassArchetypeChunk& SourceChunk, int32 OriginalIndexWithinChunk, int32 ElementsNum);
 	void MoveFragmentsToNewLocationInternal(FTransientChunkLocation Target, const FTransientChunkLocation Source, const int32 NumberToMove);
 	void ConfigureFragments(const FMassEntityManager& EntityManager);
 
@@ -578,15 +775,26 @@ inline FMassArchetypeCompositionDescriptor FMassArchetypeData::GetCompositionDes
 template<typename TStruct>
 void FMassArchetypeData::SetFragmentsData(const FMassEntityHandle Entity, TConstArrayView<TStruct> FragmentSources)
 {
-	FMassRawEntityInChunkData InternalIndex = MakeEntityHandle(Entity);
+	FMassEntityInChunkDataHandle EntityInChunkHandle = MakeEntityHandle(Entity);
 
 	for (const TStruct& Instance : FragmentSources)
 	{
 		const UScriptStruct* FragmentType = Instance.GetScriptStruct();
 		check(FragmentType);
-		const int32 FragmentIndex = FragmentIndexMap.FindChecked(FragmentType);
-		void* FragmentMemory = GetFragmentData(FragmentIndex, InternalIndex);
-		FragmentType->CopyScriptStruct(FragmentMemory, Instance.GetMemory());
+
+		if (UE::Mass::IsSparse(FragmentType) == false)
+		{
+			const int32 FragmentIndex = FragmentIndexMap.FindChecked(FragmentType);
+			void* FragmentMemory = GetFragmentData(FragmentIndex, EntityInChunkHandle);
+			FragmentType->CopyScriptStruct(FragmentMemory, Instance.GetMemory());
+		}
+		else
+		{
+			SparseElementsBitSet.Add(*FragmentType);
+			UE::Mass::FChunkSparseElements& SparseElements = Chunks[EntityInChunkHandle.ChunkIndex].GetSparseElements(NumEntitiesPerChunk);
+			SparseElements.Add(EntityInChunkHandle, FragmentType);
+			FStructView _ = SparseElementsStorage->AddElementInstanceToEntity(Entity, Instance);
+		}
 	}
 }
 
@@ -596,26 +804,57 @@ void FMassArchetypeData::SetFragmentData(FMassArchetypeEntityCollection::FConstE
 	check(FragmentSource.IsValid());
 	const UScriptStruct* FragmentType = FragmentSource.GetScriptStruct();
 	check(FragmentType);
-	const int32* FragmentIndex = FragmentIndexMap.Find(FragmentType);
-	if (LIKELY(FragmentIndex))
+
+	if (UE::Mass::IsSparse(FragmentType) == false)
 	{
-		const int32 FragmentTypeSize = FragmentType->GetStructureSize();
-		const uint8* FragmentSourceMemory = FragmentSource.GetMemory();
-		check(FragmentSourceMemory);
-	
-		for (FMassArchetypeChunkIterator ChunkIterator(EntityRangeContainer); ChunkIterator; ++ChunkIterator)
+		const int32* FragmentIndex = FragmentIndexMap.Find(FragmentType);
+		if (LIKELY(FragmentIndex))
 		{
-			uint8* FragmentMemory = (uint8*)FragmentConfigs[*FragmentIndex].GetFragmentData(Chunks[ChunkIterator->ChunkIndex].GetRawMemory(), ChunkIterator->SubchunkStart);
-			for (int i = ChunkIterator->Length; i; --i, FragmentMemory += FragmentTypeSize)
+			const int32 FragmentTypeSize = FragmentType->GetStructureSize();
+			const uint8* FragmentSourceMemory = FragmentSource.GetMemory();
+			check(FragmentSourceMemory);
+	
+			for (FMassArchetypeChunkIterator ChunkIterator(EntityRangeContainer); ChunkIterator; ++ChunkIterator)
 			{
-				FragmentType->CopyScriptStruct(FragmentMemory, FragmentSourceMemory);
+				uint8* FragmentMemory = (uint8*)FragmentConfigs[*FragmentIndex].GetFragmentData(Chunks[ChunkIterator->ChunkIndex].GetRawMemory(), ChunkIterator->SubchunkStart);
+				// Note that we cannot handle the copying with a single FragmentType->CopyScriptStruct call
+				// because we have a single FragmentSource instance we're using as a copy source (i.e. we don't have a source array).
+				for (int i = ChunkIterator->Length; i; --i, FragmentMemory += FragmentTypeSize)
+				{
+					FragmentType->CopyScriptStruct(FragmentMemory, FragmentSourceMemory);
+				}
 			}
 		}
+		else
+		{
+			UE_LOG(LogMass, Warning
+				, TEXT("Attempting to set value of fragment of type %s, while it's not part of the archetype's composition")
+				, *FragmentType->GetName());
+		}
 	}
 	else
 	{
-		UE_LOG(LogMass, Warning
-			, TEXT("Attempting to set value of fragment of type %s, while it's not part of the archetype's composition")
-			, *FragmentType->GetName());
+		SparseElementsBitSet.Add(*FragmentType);
+		for (FMassArchetypeChunkIterator ChunkIterator(EntityRangeContainer); ChunkIterator; ++ChunkIterator)
+		{
+			FMassArchetypeChunk& Chunk = Chunks[ChunkIterator->ChunkIndex];
+			const FMassArchetypeEntityCollection::FArchetypeEntityRange& EntityRange = *ChunkIterator;
+
+			UE::Mass::FChunkSparseElements& SparseElements = Chunk.GetSparseElements(NumEntitiesPerChunk);
+			SparseElements.BatchAdd(FragmentType, EntityRange.SubchunkStart, EntityRange.Length);
+
+			const FMassEntityHandle* EntitiesArray = Chunk.GetEntityArray(EntityListOffsetWithinChunk);
+			TArrayView<const FMassEntityHandle> EntityHandlesView(&EntitiesArray[EntityRange.SubchunkStart], EntityRange.Length);
+			SparseElementsStorage->BatchAddElementInstancesToEntities(EntityHandlesView, FragmentSource);
+		}
 	}
-}
\ No newline at end of file
+}
+inline bool FMassArchetypeData::ContainsAnySparseData() const
+{
+	return !SparseElementsBitSet.IsEmpty();
+}
+
+inline bool FMassArchetypeData::ContainsSparseElement(TNotNull<const UScriptStruct*> ElementType) const
+{
+	return SparseElementsBitSet.Contains(*ElementType);
+}
diff --git a/Engine/Source/Runtime/MassEntity/Private/MassArchetypeData.cpp b/Engine/Source/Runtime/MassEntity/Private/MassArchetypeData.cpp
--- a/Engine/Source/Runtime/MassEntity/Private/MassArchetypeData.cpp
+++ b/Engine/Source/Runtime/MassEntity/Private/MassArchetypeData.cpp
@@ -28,7 +28,82 @@ namespace UE::Mass
 			, InChunkMemorySize, SanitizedSize);
 		return SanitizedSize;
 	}
-}
+
+	//-----------------------------------------------------------------------------
+	// FChunkSparseElements
+	//-----------------------------------------------------------------------------
+	void FChunkSparseElements::Add(const FMassRawEntityInChunkData& EntityInChunkHandle, TNotNull<const UScriptStruct*> ElementType)
+	{
+		check(IsEmpty() == false);
+
+		const int32 TypeIndex = FMassElementBitSet::GetTypeIndex(*ElementType);
+		SparseElementsPresent.AddAtIndex(TypeIndex);
+
+		FMassElementBitSet& EntitySparseElements = PerEntityElements[EntityInChunkHandle.IndexWithinChunk];
+		EntitySparseElements.AddAtIndex(TypeIndex);
+	}
+
+	bool FChunkSparseElements::Remove(const FMassRawEntityInChunkData& EntityInChunkHandle, TNotNull<const UScriptStruct*> ElementType)
+	{
+		check(IsEmpty() == false);
+
+		FMassElementBitSet& EntitySparseElements = PerEntityElements[EntityInChunkHandle.IndexWithinChunk];
+
+		const int32 TypeIndex = FMassElementBitSet::GetTypeIndex(*ElementType);
+		const bool bHasElement = EntitySparseElements.IsBitSet(TypeIndex);
+		EntitySparseElements.RemoveAtIndex(TypeIndex);
+		return bHasElement;
+	}
+
+	void FChunkSparseElements::BatchAdd(TNotNull<const UScriptStruct*> ElementType, const int32 SubchunkStart, const int32 RangeLength)
+	{
+		check(IsEmpty() == false);
+		check(SubchunkStart >= 0 && SubchunkStart + RangeLength <= PerEntityElements.Num());
+
+		const int32 TypeIndex = FMassElementBitSet::GetTypeIndex(*ElementType);
+		for (int32 EntityIndex = SubchunkStart; EntityIndex < SubchunkStart + RangeLength; ++EntityIndex)
+		{
+			PerEntityElements[EntityIndex].AddAtIndex(TypeIndex);
+		}
+		SparseElementsPresent.AddAtIndex(TypeIndex);
+	}
+
+	void FChunkSparseElements::BatchRemove(TNotNull<const UScriptStruct*> ElementType, const int32 SubchunkStart, const int32 RangeLength)
+	{
+		check(IsEmpty() == false);
+		check(SubchunkStart >= 0 && SubchunkStart + RangeLength <= PerEntityElements.Num());
+
+		const int32 TypeIndex = FMassElementBitSet::GetTypeIndex(*ElementType);
+		for (int32 EntityIndex = SubchunkStart; EntityIndex < SubchunkStart + RangeLength; ++EntityIndex)
+		{
+			PerEntityElements[EntityIndex].RemoveAtIndex(TypeIndex);
+		}
+	}
+
+	bool FChunkSparseElements::Contains(const FMassRawEntityInChunkData& EntityInChunkHandle, TNotNull<const UScriptStruct*> ElementType) const
+	{
+		return PerEntityElements.IsValidIndex(EntityInChunkHandle.IndexWithinChunk)
+			&& PerEntityElements[EntityInChunkHandle.IndexWithinChunk].Contains(*ElementType);
+	}
+
+	void FChunkSparseElements::MoveElementsFrom(FChunkSparseElements& OriginalSparseElements, const int32 OriginalIndexWithinChunk, const int32 NewIndexWithinChunk, const int32 Num)
+	{
+		check(OriginalSparseElements.PerEntityElements.IsValidIndex(OriginalIndexWithinChunk)
+			&& OriginalIndexWithinChunk + Num <= OriginalSparseElements.PerEntityElements.Num()
+			&& PerEntityElements.IsValidIndex(NewIndexWithinChunk)
+			&& NewIndexWithinChunk + Num <= PerEntityElements.Num());
+
+		int32 OriginalIndex = OriginalIndexWithinChunk;
+		int32 NewIndex = NewIndexWithinChunk;
+		for (int32 Iterations = 0; Iterations < Num; ++Iterations, ++OriginalIndex, ++NewIndex)
+		{
+			FMassElementBitSet&& Elements = MoveTemp(OriginalSparseElements.PerEntityElements[OriginalIndex]);
+
+			SparseElementsPresent += Elements;
+			PerEntityElements[NewIndex] = MoveTemp(Elements);
+		}
+	}
+} // namespace UE::Mass
 
 //-----------------------------------------------------------------------------
 // FMassArchetypeData
@@ -52,19 +127,8 @@ void FMassArchetypeData::ForEachFragmentType(TFunction< void(const UScriptStruct
 	}
 }
 
-bool FMassArchetypeData::HasFragmentType(const UScriptStruct* FragmentType) const
-{
-	return UE::Mass::IsA<FMassFragment>(FragmentType) 
-		&& CompositionBitSet.Contains(*FragmentType);
-}
-
-bool FMassArchetypeData::HasTagType(const UScriptStruct* TagType) const
-{
-	return UE::Mass::IsA<FMassTag>(TagType)
-		&& CompositionBitSet.Contains(*TagType);
-}
-
-void FMassArchetypeData::Initialize(const FMassEntityManager& EntityManager, const FMassElementBitSet& InCompositionBitSet, const uint32 ArchetypeDataVersion)
+void FMassArchetypeData::Initialize(const FMassEntityManager& EntityManager, const FMassElementBitSet& InCompositionBitSet
+	, const uint32 ArchetypeDataVersion, UE::Mass::FSparseElementsStorage& InSparseElementsStorage)
 {
 	if (!ensureMsgf(Chunks.Num() == 0, TEXT("Trying to re-initialize non-empty Mass Archetype is not supported")))
 	{
@@ -76,6 +140,7 @@ void FMassArchetypeData::Initialize(const FMassEntityManager& EntityManager, con
 	}
 
 	CreatedArchetypeDataVersion = ArchetypeDataVersion;
+	SparseElementsStorage = &InSparseElementsStorage;
 
 	CompositionBitSet = InCompositionBitSet;
 
@@ -99,7 +164,8 @@ void FMassArchetypeData::Initialize(const FMassEntityManager& EntityManager, con
 }
 
 void FMassArchetypeData::InitializeWithSimilar(const FMassEntityManager& EntityManager, const FMassArchetypeData& BaseArchetype
-	, FMassElementBitSet&& InCompositionBitSet, const UE::Mass::FArchetypeGroups& InGroups, const uint32 ArchetypeDataVersion)
+	, FMassElementBitSet&& InCompositionBitSet, const UE::Mass::FArchetypeGroups& InGroups, const uint32 ArchetypeDataVersion
+	, UE::Mass::FSparseElementsStorage& InSparseElementsStorage)
 {
 	checkf(IsInitialized() == false, TEXT("Trying to InitializeWithSimilar but this archetype has already been initialized"));
 
@@ -110,6 +176,7 @@ void FMassArchetypeData::InitializeWithSimilar(const FMassEntityManager& EntityM
 	checkf(bAddsOnlyTagsAndElements && bNothingElse, TEXT("NewComposition is expected to differ from the BaseArchetype's composition just by tags and fragments"))
 
 	CreatedArchetypeDataVersion = ArchetypeDataVersion;
+	SparseElementsStorage = &InSparseElementsStorage;
 
 	CompositionBitSet = MoveTemp(InCompositionBitSet);
 	if (CompositionBitSet.Get<FMassFragmentBitSet>() != BaseArchetype.CompositionBitSet.Get<FMassFragmentBitSet>())
@@ -219,14 +286,61 @@ int32 FMassArchetypeData::AddEntityInternal(FMassEntityHandle Entity, const FMas
 	return AbsoluteIndex;
 }
 
+FStructView FMassArchetypeData::AddSparseElementToEntity(FMassEntityHandle EntityHandle, TNotNull<const UScriptStruct*> ElementType)
+{
+	checkf(UE::Mass::IsSparse(ElementType), TEXT("%s is not a sparse element"), *ElementType->GetName());
 
-/// ------------ /// ------------ /// ------------ /// ------------ /// ------------
-namespace UE::Mass
+	FMassEntityInChunkDataHandle EntityInChunkHandle = MakeEntityHandle(EntityHandle.Index);
+	UE::Mass::FChunkSparseElements& SparseElements = Chunks[EntityInChunkHandle.ChunkIndex].GetSparseElements(NumEntitiesPerChunk);
+
+	SparseElementsBitSet.Add(*ElementType);
+	SparseElements.Add(EntityInChunkHandle, ElementType);
+
+	// Sparse elements storage stores only fragments
+	if (UE::Mass::IsA<FMassTag>(ElementType))
+	{
+		return {};
+	}
+
+	CA_ASSUME(SparseElementsStorage);
+	return SparseElementsStorage->AddElementToEntity(EntityHandle, ElementType);
+}
+
+bool FMassArchetypeData::RemoveSparseElementFromEntity(FMassEntityHandle EntityHandle, TNotNull<const UScriptStruct*> ElementType)
 {
-	
-} // namespace UE::Mass
-/// ------------ /// ------------ /// ------------ /// ------------ /// ------------ 
+	FMassEntityInChunkDataHandle EntityInChunkHandle = MakeEntityHandle(EntityHandle.Index);
+	UE::Mass::FChunkSparseElements& SparseElements = Chunks[EntityInChunkHandle.ChunkIndex].GetSparseElements(NumEntitiesPerChunk);
 
+	if (UE::Mass::IsA<FMassTag>(ElementType) == false)
+	{
+		CA_ASSUME(SparseElementsStorage);
+		SparseElementsStorage->RemoveElementFromEntity(EntityHandle, ElementType);
+	}
+	return SparseElements.Remove(EntityInChunkHandle, ElementType);
+}
+
+FStructView FMassArchetypeData::GetMutableSparseElementDataForEntity(TNotNull<const UScriptStruct*> ElementType, const FMassEntityHandle EntityHandle) const
+{
+	check(UE::Mass::IsA<FMassTag>(ElementType) == false);
+	CA_ASSUME(SparseElementsStorage);
+	return SparseElementsStorage->GetMutableElementDataForEntity(EntityHandle, ElementType);
+}
+
+FConstStructView FMassArchetypeData::GetSparseElementDataForEntity(TNotNull<const UScriptStruct*> ElementType, const FMassEntityHandle EntityHandle) const
+{
+	check(UE::Mass::IsA<FMassTag>(ElementType) == false);
+	CA_ASSUME(SparseElementsStorage);
+	return SparseElementsStorage->GetElementDataForEntity(EntityHandle, ElementType);
+}
+
+FMassElementBitSet FMassArchetypeData::GetSparseElementsBitSetForEntity(const FMassEntityHandle EntityHandle) const
+{
+	FMassEntityInChunkDataHandle EntityInChunkHandle = MakeEntityHandle(EntityHandle.Index);
+	checkSlow(Chunks.IsValidIndex(EntityInChunkHandle.ChunkIndex));
+	return Chunks[EntityInChunkHandle.ChunkIndex].HasSparseElements()
+		? Chunks[EntityInChunkHandle.ChunkIndex].GetSparseElementsUnsafe().GetElementsForEntity(EntityInChunkHandle)
+		: FMassElementBitSet();
+}
 
 FMassArchetypeChunk& FMassArchetypeData::GetOrAddChunk(const FMassArchetypeSharedFragmentValues& SharedFragmentValues, int32& OutAbsoluteIndex, int32& OutIndexWithinChunk)
 {
@@ -298,6 +412,12 @@ void FMassArchetypeData::RemoveEntity(FMassEntityHandle Entity)
 		FragmentConfig.FragmentType->DestroyStruct(DyingFragmentPtr);
 	}
 
+	if (Chunk.HasSparseElements())
+	{
+		FMassElementBitSet& ElementsBitSet = Chunk.GetSparseElementsUnsafe().GetElementsForEntityUnsafe(IndexWithinChunk);
+		SparseElementsStorage->RemoveEntity(Entity, ElementsBitSet);
+	}
+
 	RemoveEntityInternal(AbsoluteIndex);
 }
 
@@ -324,6 +444,12 @@ void FMassArchetypeData::RemoveEntityInternal(const int32 AbsoluteIndex)
 			FMemory::Memcpy(DyingFragmentPtr, MovingFragmentPtr, FragmentConfig.FragmentType->GetStructureSize());
 		}
 
+		if (Chunk.HasSparseElements())
+		{
+			UE::Mass::FChunkSparseElements& ChunkSparseElements = Chunk.GetSparseElementsUnsafe();
+			ChunkSparseElements.GetElementsForEntityUnsafe(IndexWithinChunk) = MoveTemp(ChunkSparseElements.GetElementsForEntityUnsafe(IndexToSwapFrom));
+		}
+
 		// Update the entity table and map
 		const FMassEntityHandle EntityBeingSwapped = Chunk.GetEntityArrayElementRef(EntityListOffsetWithinChunk, IndexToSwapFrom);
 		Chunk.GetEntityArrayElementRef(EntityListOffsetWithinChunk, IndexWithinChunk) = EntityBeingSwapped;
@@ -361,6 +487,16 @@ void FMassArchetypeData::BatchDestroyEntityChunks(FMassArchetypeEntityCollection
 		FMassEntityHandle* DyingEntityPtr = &Chunk.GetEntityArrayElementRef(EntityListOffsetWithinChunk, EntityRange.SubchunkStart);
 		OutEntitiesRemoved.Append(DyingEntityPtr, RangeLength);
 
+		if (Chunk.HasSparseElements())
+		{
+			for (int32 RelativeIndex = 0; RelativeIndex < RangeLength; ++RelativeIndex)
+			{
+				FMassElementBitSet& ElementsBitSet = Chunk.GetSparseElementsUnsafe().GetElementsForEntityUnsafe(EntityRange.SubchunkStart + RelativeIndex);
+				SparseElementsStorage->RemoveEntity(DyingEntityPtr[RelativeIndex], ElementsBitSet);
+				ElementsBitSet.Reset();
+			}
+		}
+
 		for (const FMassArchetypeFragmentConfig& FragmentConfig : FragmentConfigs)
 		{
 			// Destroy the fragment data
@@ -384,10 +520,20 @@ void FMassArchetypeData::BatchDestroyEntityChunks(FMassArchetypeEntityCollection
 	}
 }
 
-bool FMassArchetypeData::HasFragmentDataForEntity(const UScriptStruct* FragmentType, const int32 EntityIndex) const
+bool FMassArchetypeData::HasSparseElementForEntity(TNotNull<const UScriptStruct*> ElementType, const FMassEntityHandle EntityHandle) const
 {
-	// @todo check if EntityIndex is disabled, or it has the FragmentType disabled.
-	return HasFragmentType(FragmentType);
+	const FMassEntityInChunkDataHandle EntityInChunkHandle = MakeEntityHandle(EntityHandle.Index);
+	return HasSparseElementForEntity(ElementType, EntityInChunkHandle);
+}
+
+bool FMassArchetypeData::HasSparseElementForEntity(TNotNull<const UScriptStruct*> ElementType, const FMassEntityInChunkDataHandle& EntityInChunkHandle) const
+{
+	checkf(Chunks.IsValidIndex(EntityInChunkHandle.ChunkIndex), TEXT("%hs: Invalid chunk index (%d out of %d)")
+		, __FUNCTION__, EntityInChunkHandle.ChunkIndex, Chunks.Num());
+
+	return SparseElementsBitSet.Contains(*ElementType)
+		&& Chunks[EntityInChunkHandle.ChunkIndex].HasSparseElements()
+		&& Chunks[EntityInChunkHandle.ChunkIndex].GetSparseElementsUnsafe().Contains(EntityInChunkHandle, ElementType);
 }
 
 void* FMassArchetypeData::GetFragmentDataForEntityChecked(TNotNull<const UScriptStruct*> FragmentType, int32 EntityIndex) const
@@ -399,22 +545,45 @@ void* FMassArchetypeData::GetFragmentDataForEntityChecked(TNotNull<const UScript
 
 void* FMassArchetypeData::GetFragmentDataForEntity(TNotNull<const UScriptStruct*> FragmentType, int32 EntityIndex) const
 {
+	checkf(UE::Mass::IsA<FMassFragment>(FragmentType), TEXT("%s is not a fragment type"), *FragmentType->GetName());
+
 	if (const int32* FragmentIndex = FragmentIndexMap.Find(FragmentType))
 	{
 		FMassRawEntityInChunkData InternalIndex = MakeEntityHandle(EntityIndex);
 		// failing the below Find means given entity's archetype is missing given FragmentType
 		return GetFragmentData(*FragmentIndex, InternalIndex);
 	}
+
+	ensureMsgf(UE::Mass::IsSparse(FragmentType) == false
+		, TEXT("Trying to fetch sparse fragment data (%s) from the archetype. Use FSparseElementsStorage instead."), *FragmentType->GetName());
 	return nullptr;
 }
 
+bool FMassArchetypeData::DoesEntityMatchSparseComposition(const FMassEntityInChunkDataHandle& EntityInChunkHandle
+	, const FMassElementBitSet& RequiredAllSparseElements, const FMassElementBitSet& RequiredAnySparseElements, const FMassElementBitSet& RequiredNoneSparseElements) const
+{
+	checkf(Chunks.IsValidIndex(EntityInChunkHandle.ChunkIndex), TEXT("%hs: Invalid chunk index (%d out of %d)")
+		, __FUNCTION__, EntityInChunkHandle.ChunkIndex, Chunks.Num());
+
+	if (Chunks[EntityInChunkHandle.ChunkIndex].HasSparseElements())
+	{
+		const FMassElementBitSet& ElementsCollection = Chunks[EntityInChunkHandle.ChunkIndex].GetSparseElementsUnsafe().GetElementsForEntity(EntityInChunkHandle);
+		return (RequiredAllSparseElements.IsEmpty() || ElementsCollection.HasAll(RequiredAllSparseElements))
+			&& (RequiredAnySparseElements.IsEmpty() || ElementsCollection.HasAny(RequiredAnySparseElements))
+			&& ElementsCollection.HasNone(RequiredNoneSparseElements);
+	}
+
+	// no sparse elements, so the entity "matches" only if the requirements are empty
+	return RequiredAllSparseElements.IsEmpty() && RequiredAnySparseElements.IsEmpty();
+}
+
 void FMassArchetypeData::SetSharedFragmentsData(const FMassEntityHandle Entity, TConstArrayView<FSharedStruct> SharedFragmentValueOverrides)
 {
 	// Gets the current chunk where the entity is located
 	const int32 OldAbsoluteIndex = EntityMap.FindChecked(Entity.Index);
 	const int32 OldChunkIndex = OldAbsoluteIndex / NumEntitiesPerChunk;
 	const int32 OldIndexWithinChunk = OldAbsoluteIndex - (NumEntitiesPerChunk * OldChunkIndex);
-	const FMassArchetypeChunk& OldChunk = Chunks[OldChunkIndex];
+	FMassArchetypeChunk& OldChunk = Chunks[OldChunkIndex];
 
 	// Gets or adds a new chunk that will hold the new entity with the new shared values
 	FMassArchetypeSharedFragmentValues NewSharedFragmentValues(OldChunk.GetSharedFragmentValues());
@@ -437,6 +606,13 @@ void FMassArchetypeData::SetSharedFragmentsData(const FMassEntityHandle Entity,
 		// Move the current entity fragments into the new chunk
 		MoveFragmentsToNewLocationInternal({ OldChunk.GetRawMemory(), OldIndexWithinChunk }, { NewChunk.GetRawMemory(), NewIndexWithinChunk }, 1);
 
+		if (OldChunk.HasSparseElements())
+		{
+			UE::Mass::FChunkSparseElements& OldSparseElements = OldChunk.GetSparseElementsUnsafe();
+			UE::Mass::FChunkSparseElements& NewSparseElements = NewChunk.GetSparseElements(NumEntitiesPerChunk);
+			NewSparseElements.MoveElementsFrom(OldSparseElements, OldIndexWithinChunk, NewIndexWithinChunk, 1);
+		}
+
 		// Clean up the old chunk
 		RemoveEntityInternal(OldAbsoluteIndex);
 	}
@@ -459,7 +635,7 @@ void FMassArchetypeData::MoveEntityToAnotherArchetype(const FMassEntityHandle En
 	UE_TRACE_MASS_ARCHETYPE_CREATED(NewArchetype)
 	UE_TRACE_MASS_ENTITY_MOVED(Entity, NewArchetype)
 
-	MoveFragmentsToAnotherArchetypeInternal(NewArchetype, { NewChunk.GetRawMemory(), NewIndexWithinChunk }, { Chunk.GetRawMemory(), IndexWithinChunk }, /*Count=*/1);
+	MoveFragmentsToAnotherArchetypeInternal(NewArchetype, NewChunk, NewIndexWithinChunk, Chunk, IndexWithinChunk, /*Count=*/1);
 
 	RemoveEntityInternal(AbsoluteIndex);
 }
@@ -472,7 +648,7 @@ void FMassArchetypeData::ExecuteFunction(FMassExecutionContext& RunContext, cons
 	}
 
 	// @todo do we really want users to check composition of the archetype being processed at the moment?
-	RunContext.StoreCurrentArchetypeComposition(CompositionBitSet);
+	RunContext.SetCurrentArchetype(*this);
 #if WITH_MASSENTITY_DEBUG
 	RunContext.DebugSetColor(DebugColor);
 #endif // WITH_MASSENTITY_DEBUG
@@ -495,11 +671,11 @@ void FMassArchetypeData::ExecuteFunction(FMassExecutionContext& RunContext, cons
 
 			checkf((ChunkIterator->SubchunkStart + SubchunkLength) <= Chunk.GetNumInstances() && SubchunkLength > 0, TEXT("Invalid subchunk, it is going over the number of instances in the chunk or it is empty."));
 
-			RunContext.SetCurrentChunkSerialModificationNumber(Chunk.GetSerialModificationNumber());
 			BindChunkFragmentRequirements(RunContext, RequirementMapping.ChunkFragments, Chunk);
 			RunContext.BindLinkedEntity();
 
-			if (!ChunkCondition || ChunkCondition(RunContext))
+			if (RunContext.SetCurrentChunk(ChunkIterator->ChunkIndex, Chunk) 
+				&& (!ChunkCondition || ChunkCondition(RunContext)))
 			{
 				BindEntityRequirements(RunContext, RequirementMapping.EntityFragments, Chunk, ChunkIterator->SubchunkStart, SubchunkLength);
 				Function(RunContext);
@@ -515,7 +691,7 @@ void FMassArchetypeData::ExecuteFunction(FMassExecutionContext& RunContext, cons
 		return;
 	}
 
-	RunContext.StoreCurrentArchetypeComposition(CompositionBitSet);
+	RunContext.SetCurrentArchetype(*this);
 #if WITH_MASSENTITY_DEBUG
 	RunContext.DebugSetColor(DebugColor);
 #endif // WITH_MASSENTITY_DEBUG
@@ -547,11 +723,11 @@ void FMassArchetypeData::ExecuteFunction(FMassExecutionContext& RunContext, cons
 				BindSharedFragmentRequirements(RunContext, Chunk.GetMutableSharedFragmentValues(), RequirementMapping.SharedFragments);
 			}
 
-			RunContext.SetCurrentChunkSerialModificationNumber(Chunk.GetSerialModificationNumber());
 			BindChunkFragmentRequirements(RunContext, RequirementMapping.ChunkFragments, Chunk);
 			RunContext.BindLinkedEntity();
 
-			if (!ChunkCondition || ChunkCondition(RunContext))
+			if (RunContext.SetCurrentChunk(ChunkIndex, Chunk)
+				&& (!ChunkCondition || ChunkCondition(RunContext)))
 			{
 				BindEntityRequirements(RunContext, RequirementMapping.EntityFragments, Chunk, 0, Chunk.GetNumInstances());
 				Function(RunContext);
@@ -579,8 +755,7 @@ void FMassArchetypeData::ExecutionFunctionForChunk(FMassExecutionContext& RunCon
 		BindConstSharedFragmentRequirements(RunContext, Chunk.GetSharedFragmentValues(), RequirementMapping.ConstSharedFragments);
 		BindSharedFragmentRequirements(RunContext, Chunk.GetMutableSharedFragmentValues(), RequirementMapping.SharedFragments);
 
-		RunContext.StoreCurrentArchetypeComposition(CompositionBitSet);
-		RunContext.SetCurrentChunkSerialModificationNumber(Chunk.GetSerialModificationNumber());
+		RunContext.SetCurrentArchetype(*this);
 #if WITH_MASSENTITY_DEBUG
 		RunContext.DebugSetColor(DebugColor);
 #endif // WITH_MASSENTITY_DEBUG
@@ -588,7 +763,8 @@ void FMassArchetypeData::ExecutionFunctionForChunk(FMassExecutionContext& RunCon
 		BindChunkFragmentRequirements(RunContext, RequirementMapping.ChunkFragments, Chunk);
 		RunContext.BindLinkedEntity();
 
-		if (!ChunkCondition || ChunkCondition(RunContext))
+		if (RunContext.SetCurrentChunk(EntityRange.ChunkIndex, Chunk)
+			&& (!ChunkCondition || ChunkCondition(RunContext)))
 		{
 			BindEntityRequirements(RunContext, RequirementMapping.EntityFragments, Chunk, EntityRange.SubchunkStart, RangeLength);
 			Function(RunContext);
@@ -656,6 +832,13 @@ int32 FMassArchetypeData::CompactEntities(const double TimeAllowed)
 			MoveFragmentsToNewLocationInternal({ChunkToFill->GetRawMemory(), ToIndex}, {ChunkToEmpty->GetRawMemory(), FromIndex}
 				, NumberOfEntitiesToMove);
 
+			if (ChunkToEmpty->HasSparseElements())
+			{
+				UE::Mass::FChunkSparseElements& OldSparseElements = ChunkToEmpty->GetSparseElementsUnsafe();
+				UE::Mass::FChunkSparseElements& NewSparseElements = ChunkToFill->GetSparseElements(NumEntitiesPerChunk);
+				NewSparseElements.MoveElementsFrom(OldSparseElements, FromIndex, ToIndex, NumberOfEntitiesToMove);
+			}
+
 			FMassEntityHandle* FromEntity = &ChunkToEmpty->GetEntityArrayElementRef(EntityListOffsetWithinChunk, FromIndex);
 			FMassEntityHandle* ToEntity = &ChunkToFill->GetEntityArrayElementRef(EntityListOffsetWithinChunk, ToIndex);
 			FMemory::Memcpy(ToEntity, FromEntity, NumberOfEntitiesToMove * sizeof(FMassEntityHandle));
@@ -1217,7 +1400,7 @@ void FMassArchetypeData::BatchMoveEntitiesToAnotherArchetype(const FMassArchetyp
 			}
 
 			FMassArchetypeChunk& NewChunk = NewArchetype.Chunks[ResultSubChunk.ChunkIndex];
-			MoveFragmentsToAnotherArchetypeInternal(NewArchetype, {NewChunk.GetRawMemory(), ResultSubChunk.SubchunkStart}, {Chunk.GetRawMemory(), IndexWithinChunk}, ResultSubChunk.Length);
+			MoveFragmentsToAnotherArchetypeInternal(NewArchetype, NewChunk, ResultSubChunk.SubchunkStart, Chunk, IndexWithinChunk, ResultSubChunk.Length);
 
 			NumberMoved += ResultSubChunk.Length;
 
@@ -1342,7 +1525,13 @@ void FMassArchetypeData::BatchRemoveEntitiesInternal(const int32 ChunkIndex, con
 		checkf((StartIndexWithinChunk + NumberToMove - 1) < SwapStartIndex, TEXT("Remove and Move ranges overlap"));
 
 		MoveFragmentsToNewLocationInternal({ Chunk.GetRawMemory(), StartIndexWithinChunk }, { Chunk.GetRawMemory(), SwapStartIndex }, NumberToMove);
-		
+
+		if (Chunk.HasSparseElements())
+		{
+			UE::Mass::FChunkSparseElements& SparseElements = Chunk.GetSparseElementsUnsafe();
+			SparseElements.MoveElementsFrom(SparseElements, SwapStartIndex, StartIndexWithinChunk, NumberToMove);
+		}
+
 		// Update the entity table and map
 		const FMassEntityHandle* MovingEntityPtr = &Chunk.GetEntityArrayElementRef(EntityListOffsetWithinChunk, SwapStartIndex);
 		int32 AbsoluteIndex = ChunkIndex * NumEntitiesPerChunk + StartIndexWithinChunk;
@@ -1364,9 +1553,12 @@ void FMassArchetypeData::BatchRemoveEntitiesInternal(const int32 ChunkIndex, con
 	}
 }
 
-void FMassArchetypeData::MoveFragmentsToAnotherArchetypeInternal(FMassArchetypeData& TargetArchetype, FMassArchetypeData::FTransientChunkLocation Target
-	, const FMassArchetypeData::FTransientChunkLocation Source, const int32 ElementsNum)
+void FMassArchetypeData::MoveFragmentsToAnotherArchetypeInternal(FMassArchetypeData& TargetArchetype, FMassArchetypeChunk& TargetChunk, const int32 NewIndexWithinChunk
+	, FMassArchetypeChunk& SourceChunk, const int32 OriginalIndexWithinChunk, const int32 ElementsNum)
 {
+	FTransientChunkLocation Target(TargetChunk.GetRawMemory(), NewIndexWithinChunk);
+	FTransientChunkLocation Source(SourceChunk.GetRawMemory(), OriginalIndexWithinChunk);
+
 	// for every TargetArchetype's fragment see if it was in the old archetype as well and if so copy it's value. 
 	// If not then initialize the fragment.
 	for (const FMassArchetypeFragmentConfig& TargetFragmentConfig : TargetArchetype.FragmentConfigs)
@@ -1400,6 +1592,17 @@ void FMassArchetypeData::MoveFragmentsToAnotherArchetypeInternal(FMassArchetypeD
 			FragmentConfig.FragmentType->DestroyStruct(DyingFragmentPtr, ElementsNum);
 		}
 	}
+
+	// move sparse elements
+	if (SourceChunk.HasSparseElements())
+	{
+		UE::Mass::FChunkSparseElements& SourceSparseElements = SourceChunk.GetSparseElements(NumEntitiesPerChunk);
+		// calling GetSparseElements to make sure FChunkSparseElements gets created with the right size (FChunkSparseElements)
+		UE::Mass::FChunkSparseElements& TargetElements = TargetChunk.GetSparseElements(TargetArchetype.NumEntitiesPerChunk);
+		TargetElements.MoveElementsFrom(SourceSparseElements, OriginalIndexWithinChunk, NewIndexWithinChunk, ElementsNum);
+
+		TargetArchetype.SparseElementsBitSet += TargetElements.GetSparseElementsPresent();
+	}
 }
 
 FORCEINLINE void FMassArchetypeData::MoveFragmentsToNewLocationInternal(FMassArchetypeData::FTransientChunkLocation Target, const FMassArchetypeData::FTransientChunkLocation Source, const int32 NumberToMove)
@@ -1431,17 +1634,85 @@ void FMassArchetypeData::BatchSetFragmentValues(TConstArrayView<FMassArchetypeEn
 			const UScriptStruct* FragmentType = FragmentPayload.GetScriptStruct();
 			check(FragmentType);
 
-			const int32 FragmentIndex = FragmentIndexMap.FindChecked(FragmentType);
-			void* Dst = FragmentConfigs[FragmentIndex].GetFragmentData(Chunk.GetRawMemory(), EntityRange.SubchunkStart);
-			const void* Src = FragmentPayload.GetDataAt(EntitiesHandled);
+			if (UE::Mass::IsSparse(FragmentType) == false)
+			{
+				const int32 FragmentIndex = FragmentIndexMap.FindChecked(FragmentType);
+				void* Dst = FragmentConfigs[FragmentIndex].GetFragmentData(Chunk.GetRawMemory(), EntityRange.SubchunkStart);
+				const void* Src = FragmentPayload.GetDataAt(EntitiesHandled);
 
-			FragmentType->CopyScriptStruct(Dst, Src, RangeLength);
+				FragmentType->CopyScriptStruct(Dst, Src, RangeLength);
+			}
+			else
+			{
+				SparseElementsBitSet.Add(*FragmentType);
+				UE::Mass::FChunkSparseElements& SparseElements = Chunk.GetSparseElements(NumEntitiesPerChunk);
+				SparseElements.BatchAdd(FragmentType, EntityRange.SubchunkStart, RangeLength);
+
+				const FMassEntityHandle* EntitiesArray = Chunk.GetEntityArray(EntityListOffsetWithinChunk);
+				TArrayView<const FMassEntityHandle> EntityHandlesView(&EntitiesArray[EntityRange.SubchunkStart], RangeLength);
+				FStructArrayView FragmentPayloadSlice(*FragmentType, FragmentPayload.GetDataAt(EntitiesHandled), EntityHandlesView.Num());
+				SparseElementsStorage->BatchAddElementInstancesToEntities(EntityHandlesView, FragmentPayloadSlice);
+			}
 		}
 
 		EntitiesHandled += RangeLength;
 	}
 }
 
+void FMassArchetypeData::BatchAddSparseElementToEntities(const FMassArchetypeEntityCollection& EntityCollection, TNotNull<const UScriptStruct*> ElementType)
+{
+	checkf(UE::Mass::IsSparse(ElementType), TEXT("%s is not a sparse element type"), *ElementType->GetName());
+
+	const bool bIsTag = UE::Mass::IsA<FMassTag>(ElementType);
+	SparseElementsBitSet.Add(*ElementType);
+
+	TArray<FMassArchetypeEntityCollection::FArchetypeEntityRange> Subchunks(EntityCollection.GetRanges());
+	for (const FMassArchetypeEntityCollection::FArchetypeEntityRange EntityRange : Subchunks)
+	{
+		FMassArchetypeChunk& Chunk = Chunks[EntityRange.ChunkIndex];
+		const int32 RangeLength = CalculateRangeLength(EntityRange, Chunk);
+
+		UE::Mass::FChunkSparseElements& SparseElements = Chunk.GetSparseElements(NumEntitiesPerChunk);
+		SparseElements.BatchAdd(ElementType, EntityRange.SubchunkStart, RangeLength);
+
+		if (!bIsTag)
+		{
+			const FMassEntityHandle* EntitiesArray = Chunk.GetEntityArray(EntityListOffsetWithinChunk);
+			TArrayView<const FMassEntityHandle> EntityHandlesView(&EntitiesArray[EntityRange.SubchunkStart], RangeLength);
+			SparseElementsStorage->BatchAddElementToEntities(EntityHandlesView, ElementType);
+		}
+	}
+}
+
+void FMassArchetypeData::BatchRemoveSparseElementFromEntities(const FMassArchetypeEntityCollection& EntityCollection, TNotNull<const UScriptStruct*> ElementType)
+{
+	checkf(UE::Mass::IsSparse(ElementType), TEXT("%s is not a sparse element type"), *ElementType->GetName());
+
+	if (SparseElementsBitSet.Contains(*ElementType) == false)
+	{
+		return;
+	}
+
+	const bool bIsTag = UE::Mass::IsA<FMassTag>(ElementType);
+
+	TArray<FMassArchetypeEntityCollection::FArchetypeEntityRange> Subchunks(EntityCollection.GetRanges());
+	for (const FMassArchetypeEntityCollection::FArchetypeEntityRange EntityRange : Subchunks)
+	{
+		FMassArchetypeChunk& Chunk = Chunks[EntityRange.ChunkIndex];
+		const int32 RangeLength = CalculateRangeLength(EntityRange, Chunk);
+
+		UE::Mass::FChunkSparseElements& SparseElements = Chunk.GetSparseElements(NumEntitiesPerChunk);
+		SparseElements.BatchRemove(ElementType, EntityRange.SubchunkStart, RangeLength);
+
+		if (!bIsTag)
+		{
+			const FMassEntityHandle* EntitiesArray = Chunk.GetEntityArray(EntityListOffsetWithinChunk);
+			TArrayView<const FMassEntityHandle> EntityHandlesView(&EntitiesArray[EntityRange.SubchunkStart], RangeLength);
+			SparseElementsStorage->BatchRemoveElementFromEntities(EntityHandlesView, ElementType);
+		}
+	}
+}
+
 bool FMassArchetypeData::IsEquivalent(const FMassElementBitSet& OtherCompositionBitSet, const UE::Mass::FArchetypeGroups& OtherGroups) const
 {
 	return CompositionBitSet.IsEquivalent(OtherCompositionBitSet) && Groups == OtherGroups;
diff --git a/Engine/Source/Runtime/MassEntity/Private/MassElement.cpp b/Engine/Source/Runtime/MassEntity/Private/MassElement.cpp
--- a/Engine/Source/Runtime/MassEntity/Private/MassElement.cpp
+++ b/Engine/Source/Runtime/MassEntity/Private/MassElement.cpp
@@ -11,6 +11,30 @@ namespace UE::Mass
 	TStaticArray<FElementBitSetBase, static_cast<uint8>(EElementType::MAX)> FElementBitSet::ElementTypeBitArrays;
 	FElementBitSet FElementBitSet::AllSharedFragmentsBitSet;
 	FElementBitSet FElementBitSet::AllFragmentsAndTagsBitSet;
-	bool FElementBitSet::bTypeBitSetsDirty = false;
+	FElementBitSet FElementBitSet::AllSparseElementsBitSet;
 	FDelegateHandle FElementBitSet::OnTypeRegisteredDelegateHandle;
+
+	void FElementBitSet::OnTypeRegistered(const TNotNull<const UStruct*> Type, int32 Index)
+	{
+		const EElementType ElementType = DetermineElementType(Type);
+		GetTypeBitArray(ElementType).AddAtIndex(Index);
+
+		switch (ElementType)
+		{
+		case EElementType::Fragment: [[fallthrough]];
+		case EElementType::Tag:
+			AllFragmentsAndTagsBitSet.AddAtIndex(Index);
+			if (IsSparse(Type))
+			{
+				AllSparseElementsBitSet.AddAtIndex(Index);
+			}
+			break;
+		case EElementType::ChunkFragment:
+			break;
+		case EElementType::SharedFragment: [[fallthrough]];
+		case EElementType::ConstSharedFragment:
+			AllSharedFragmentsBitSet.AddAtIndex(Index);
+			break;
+		}
+	}
 }
diff --git a/Engine/Source/Runtime/MassEntity/Private/MassEntityBuilder.cpp b/Engine/Source/Runtime/MassEntity/Private/MassEntityBuilder.cpp
--- a/Engine/Source/Runtime/MassEntity/Private/MassEntityBuilder.cpp
+++ b/Engine/Source/Runtime/MassEntity/Private/MassEntityBuilder.cpp
@@ -296,16 +296,18 @@ FMassEntityHandle FEntityBuilder::Commit()
 		EntityManager->Defer().PushCommand<FMassDeferredCreateCommand>(
 					[ReservedEntityHandle = EntityHandle, SharedFragmentValues = CachedSharedFragmentValues
 					, ArchetypeHandle = CachedArchetypeHandle, FragmentsCopy = Fragments
-					, RelationsParams = RelationsParams](FMassEntityManager& Manager)
+					, RelationsParams = RelationsParams
+					, CompositionBitSet = Composition.GetElementsBitSet()](FMassEntityManager& Manager)
 					{
 						FEntityBuilder::CreateEntityImpl(Manager, ReservedEntityHandle, ArchetypeHandle
-							, SharedFragmentValues, FragmentsCopy, RelationsParams);
+							, SharedFragmentValues, FragmentsCopy, RelationsParams, CompositionBitSet);
 					});
 	}
 	else
 	{
 		// directly create the entity, right now
-		CreateEntityImpl(*EntityManager, EntityHandle, CachedArchetypeHandle, CachedSharedFragmentValues, Fragments, RelationsParams);
+		CreateEntityImpl(*EntityManager, EntityHandle, CachedArchetypeHandle, CachedSharedFragmentValues
+			, Fragments, RelationsParams, Composition.GetElementsBitSet());
 	}
 
 	State = EState::Committed;
@@ -314,7 +316,8 @@ FMassEntityHandle FEntityBuilder::Commit()
 }
 
 void FEntityBuilder::CreateEntityImpl(FMassEntityManager& EntityManager, FMassEntityHandle EntityHandle, const FMassArchetypeHandle& ArchetypeHandle
-	, const FMassArchetypeSharedFragmentValues& SharedFragmentValues, const TConstArrayView<FInstancedStruct> Fragments, TConstArrayView<FPendingRelationParams> RelationsParams)
+	, const FMassArchetypeSharedFragmentValues& SharedFragmentValues, const TConstArrayView<FInstancedStruct> Fragments, TConstArrayView<FPendingRelationParams> RelationsParams
+	, const FMassElementBitSet& CompositionBitSet)
 {
 	// creating creation context to block observers from triggering until the
 	// values are set with SetEntityFragmentValues. The lock gets auto-released at the end of the scope
@@ -323,6 +326,13 @@ void FEntityBuilder::CreateEntityImpl(FMassEntityManager& EntityManager, FMassEn
 
 	EntityManager.BuildEntity(EntityHandle, ArchetypeHandle, SharedFragmentValues);
 	EntityManager.SetEntityFragmentValues(EntityHandle, Fragments);
+
+	FMassElementBitSet SparseElements = CompositionBitSet & FMassElementBitSet::GetAllSparseElements();
+	if (SparseElements.IsEmpty() == false)
+	{
+		EntityManager.AddSparseElementsToEntity(EntityHandle, SparseElements);
+	}
+
 	for (const FPendingRelationParams& Relation : RelationsParams)
 	{
 		switch (Relation.OtherEntityRole)
@@ -393,7 +403,7 @@ bool FEntityBuilder::SetReservedEntityHandle(const FMassEntityHandle ReservedEnt
 	return true;
 }
 
-bool FEntityBuilder::AppendDataFromEntity(const FMassEntityHandle SourceEntityHandle)
+bool FEntityBuilder::AppendDataFromEntity(const FMassEntityHandle SourceEntityHandle, EIncludeSparseElements IncludeSparseElements)
 {
 	if (!testableEnsureMsgf(EntityManager->IsEntityActive(SourceEntityHandle)
 		, TEXT("%hs expecting a valid, built entity as input"), __FUNCTION__))
@@ -403,11 +413,42 @@ bool FEntityBuilder::AppendDataFromEntity(const FMassEntityHandle SourceEntityHa
 	if (State == EState::Empty)
 	{
 		// copying is significantly more efficient (no lookups for existing data) 
-		return CopyDataFromEntity(SourceEntityHandle);
+		return CopyDataFromEntity(SourceEntityHandle, IncludeSparseElements);
 	}
 
 	InvalidateCachedData();
 
+	// handling sparse elements first in case we need to look for existing sparse fragment instances in Fragments,
+	// before we add more stuff to it
+	if (IncludeSparseElements == EIncludeSparseElements::Yes)
+	{
+		const FMassElementBitSet SparseElements = EntityManager->GetSparseElementsBitSetForEntity(SourceEntityHandle);
+		const FMassElementBitSet SparseFragments = SparseElements & FMassElementBitSet::GetTypeBitArray(UE::Mass::EElementType::Fragment);
+		for (FMassElementBitSet::FIndexIterator It = SparseFragments.GetIndexIterator(); It; ++It)
+		{
+			const UScriptStruct* SparseElementType = FMassElementBitSet::GetTypeAtIndex(*It);
+			CA_ASSUME(SparseElementType);
+			const FConstStructView FragmentInstance = EntityManager->GetSparseElementDataForEntity(SparseElementType, SourceEntityHandle);
+
+			if (Composition.GetElementsBitSet().GetAtIndex(*It))
+			{
+				// it's possible there's an instance in Fragments already
+				FInstancedStruct* ExistingStruct = Fragments.FindByPredicate([&SparseElementType](const FInstancedStruct& Element)
+					{
+						return Element.GetScriptStruct() == SparseElementType;
+					});
+				if (ExistingStruct)
+				{
+					*ExistingStruct = FragmentInstance;
+					continue;
+				}
+			}
+			
+			Fragments.Emplace(FragmentInstance);
+		}
+		Composition.GetElementsBitSet() += SparseElements;
+	}
+
 	const FMassArchetypeHandle ArchetypeHandle = EntityManager->GetArchetypeForEntity(SourceEntityHandle);
 	const FMassArchetypeCompositionDescriptor& ArchetypeComposition = EntityManager->GetArchetypeComposition(ArchetypeHandle);
 
@@ -422,9 +463,8 @@ bool FEntityBuilder::AppendDataFromEntity(const FMassEntityHandle SourceEntityHa
 	return true;
 }
 
-bool FEntityBuilder::CopyDataFromEntity(const FMassEntityHandle SourceEntityHandle)
+bool FEntityBuilder::CopyDataFromEntity(const FMassEntityHandle SourceEntityHandle, EIncludeSparseElements IncludeSparseElements)
 {
-	// @todo should we copy transient data too?
 	if (!testableEnsureMsgf(EntityManager->IsEntityActive(SourceEntityHandle)
 		, TEXT("%hs expecting a valid, built entity as input"), __FUNCTION__))
 	{
@@ -442,6 +482,22 @@ bool FEntityBuilder::CopyDataFromEntity(const FMassEntityHandle SourceEntityHand
 
 	Composition = ArchetypeComposition;
 
+	if (IncludeSparseElements == EIncludeSparseElements::Yes)
+	{
+		const FMassElementBitSet SparseElements = EntityManager->GetSparseElementsBitSetForEntity(SourceEntityHandle);
+		const FMassElementBitSet SparseFragments = SparseElements & FMassElementBitSet::GetTypeBitArray(UE::Mass::EElementType::Fragment);
+		for (FMassElementBitSet::FIndexIterator It = SparseFragments.GetIndexIterator(); It; ++It)
+		{
+			const UScriptStruct* SparseElementType = FMassElementBitSet::GetTypeAtIndex(*It);
+			CA_ASSUME(SparseElementType);
+			const FConstStructView FragmentInstance = EntityManager->GetSparseElementDataForEntity(SparseElementType, SourceEntityHandle);
+			// we can add elements without checking if instances of the given type has already been added
+			// Due to the nature of how sparse fragments are stored it's guaranteed the Fragments container doesn't contain any
+			Fragments.Emplace(FragmentInstance);
+		}
+		Composition.GetElementsBitSet() += SparseElements;
+	}
+
 	State = Composition.IsEmpty() ? EState::Empty : EState::ReadyToCommit;
 
 	return true;
diff --git a/Engine/Source/Runtime/MassEntity/Private/MassEntityManager.cpp b/Engine/Source/Runtime/MassEntity/Private/MassEntityManager.cpp
--- a/Engine/Source/Runtime/MassEntity/Private/MassEntityManager.cpp
+++ b/Engine/Source/Runtime/MassEntity/Private/MassEntityManager.cpp
@@ -23,6 +23,8 @@
 #define CHECK_SYNC_API_RETURN(ReturnValue) testableCheckfReturn(IsProcessing() == false, ReturnValue , TEXT("Synchronous API function %hs called during mass processing. Use asynchronous API instead."), __FUNCTION__)
 #define CHECK_ELEMENT(ElementType) checkf(UE::Mass::Private::IsElement(ElementType), TEXT("%hs: Only tags and fragments are considered 'elements', type provided: %s") \
 		, __FUNCTION__ , *ElementType->GetName())
+#define CHECK_SPARSE_ELEMENT(ElementType) checkf(UE::Mass::Private::IsSparseElement(ElementType), TEXT("%hs: Only sparse tags and fragments can be added or removed via this API, type provided: %s") \
+		, __FUNCTION__ , *ElementType->GetName())
 
 const FMassEntityHandle FMassEntityManager::InvalidEntity;
 
@@ -54,6 +56,12 @@ namespace UE::Mass::Private
 		return UE::Mass::IsA<FMassFragment>(ElementType)
 			|| UE::Mass::IsA<FMassTag>(ElementType);
 	}
+
+	bool IsSparseElement(TNotNull<const UScriptStruct*> ElementType)
+	{
+		return ElementType->IsChildOf(FMassSparseFragment::StaticStruct())
+			|| ElementType->IsChildOf(FMassSparseTag::StaticStruct());
+	}
 }
 
 //-----------------------------------------------------------------------------
@@ -294,7 +302,9 @@ void FMassEntityManager::Deinitialize()
 		EntityStorage.Emplace<FEmptyVariantState>();
 
 		ObserverManager.DeInitialize();
-		
+
+		SparseElementsStorage.Deinitialize();
+
 		InitializationState = EInitializationState::Deinitialized;
 	}
 	else
@@ -388,14 +398,17 @@ FMassArchetypeHandle FMassEntityManager::GetOrCreateSuitableArchetype(const FMas
 FMassArchetypeHandle FMassEntityManager::CreateArchetype(const FMassArchetypeCompositionDescriptor& Composition, const FMassArchetypeCreationParams& CreationParams)
 {
 	LLM_SCOPE_BYNAME(TEXT("Mass/EntityManager"));
-	const uint32 TypeHash = HashCombine(Composition.CalculateHash(), GetTypeHash(UE::Mass::FArchetypeGroups()));
+	/*FMassElementBitSet AndToo = Composition.GetElementsBitSet() & FMassElementBitSet::GetAllSparseElements();
+	check(AndToo.IsEmpty());*/
+	FMassElementBitSet NonSparseComposition = Composition.GetElementsBitSet().AndNot(FMassElementBitSet::GetAllSparseElements());
+	const uint32 TypeHash = HashCombine(GetTypeHash(NonSparseComposition), GetTypeHash(UE::Mass::FArchetypeGroups()));
 
 	TArray<TSharedPtr<FMassArchetypeData>>& HashRow = FragmentHashToArchetypeMap.FindOrAdd(TypeHash);
 
 	TSharedPtr<FMassArchetypeData> ArchetypeDataPtr;
 	for (const TSharedPtr<FMassArchetypeData>& Ptr : HashRow)
 	{
-		if (Ptr->IsEquivalent(Composition.GetElementsBitSet(), /*Groups=*/{}))
+		if (Ptr->IsEquivalent(NonSparseComposition, /*Groups=*/{}))
 		{
 #if WITH_MASSENTITY_DEBUG
 			// Keep track of all names for this archetype.
@@ -421,7 +434,7 @@ FMassArchetypeHandle FMassEntityManager::CreateArchetype(const FMassArchetypeCom
 
 		// Create a new archetype
 		FMassArchetypeData* NewArchetype = new FMassArchetypeData(CreationParams);
-		NewArchetype->Initialize(*this, Composition.GetElementsBitSet(), ArchetypeDataVersion);
+		NewArchetype->Initialize(*this, NonSparseComposition, ArchetypeDataVersion, SparseElementsStorage);
 		ArchetypeDataPtr = HashRow.Add_GetRef(MakeShareable(NewArchetype));
 		AllArchetypes.Add(ArchetypeDataPtr);
 		ensure(AllArchetypes.Num() == ArchetypeDataVersion);
@@ -461,14 +474,17 @@ FMassArchetypeHandle FMassEntityManager::InternalCreateSimilarArchetype(const FM
 	// be there if IDs were added and removed to this specific Groups container instance
 	checkf(Groups.IsShrunk(), TEXT("A group container with invalid trailing IDs has been passed to archetype creation - this is not expected and will cause issues. Make sure to Shrink your Groups before passing to %hs"), __FUNCTION__);
 
-	const uint32 TypeHash = HashCombine(NewComposition.CalculateHash(), GetTypeHash(Groups));
+	/*FMassElementBitSet AndToo = NewComposition.GetElementsBitSet() & FMassElementBitSet::GetAllSparseElements();
+	check(AndToo.IsEmpty());*/
+	FMassElementBitSet NonSparseComposition = NewComposition.GetElementsBitSet().AndNot(FMassElementBitSet::GetAllSparseElements());
+	const uint32 TypeHash = HashCombine(GetTypeHash(NonSparseComposition), GetTypeHash(UE::Mass::FArchetypeGroups()));
 
 	TArray<TSharedPtr<FMassArchetypeData>>& HashRow = FragmentHashToArchetypeMap.FindOrAdd(TypeHash);
 
 	TSharedPtr<FMassArchetypeData> ArchetypeDataPtr;
 	for (const TSharedPtr<FMassArchetypeData>& Ptr : HashRow)
 	{
-		if (Ptr->IsEquivalent(NewComposition.GetElementsBitSet(), Groups))
+		if (Ptr->IsEquivalent(NonSparseComposition, Groups))
 		{
 			ArchetypeDataPtr = Ptr;
 			break;
@@ -482,7 +498,7 @@ FMassArchetypeHandle FMassEntityManager::InternalCreateSimilarArchetype(const FM
 
 		// Create a new archetype
 		FMassArchetypeData* NewArchetype = new FMassArchetypeData(FMassArchetypeCreationParams(SourceArchetypeRef));
-		NewArchetype->InitializeWithSimilar(*this, SourceArchetypeRef, MoveTemp(NewComposition.GetElementsBitSet()), Groups, ArchetypeDataVersion);
+		NewArchetype->InitializeWithSimilar(*this, SourceArchetypeRef, MoveTemp(NonSparseComposition), Groups, ArchetypeDataVersion, SparseElementsStorage);
 		NewArchetype->CopyDebugNamesFrom(SourceArchetypeRef);
 
 		ArchetypeDataPtr = HashRow.Add_GetRef(MakeShareable(NewArchetype));
@@ -886,6 +902,9 @@ void FMassEntityManager::BatchDestroyEntities(TConstArrayView<FMassEntityHandle>
 
 	MASS_BREAKPOINT(UE::Mass::Debug::FBreakpoint::CheckDestroyEntityBreakpoints(InEntities));
 
+	TArray<FMassEntityHandle> ValidEntityHandles;
+	ValidEntityHandles.Reserve(InEntities.Num());
+
 	for (const FMassEntityHandle Entity : InEntities)
 	{
 		if (GetEntityStorageInterface().IsValidIndex(Entity.Index) == false)
@@ -903,6 +922,7 @@ void FMassEntityManager::BatchDestroyEntities(TConstArrayView<FMassEntityHandle>
 		{
 			ObserverManager.OnPreEntityDestroyed(Archetype->GetCompositionDescriptor(), Entity);
 			Archetype->RemoveEntity(Entity);
+			ValidEntityHandles.Add(Entity);
 		}
 		// else it's a "reserved" entity so it has not been assigned to an archetype yet, no archetype nor observers to notify
 	}
@@ -938,7 +958,7 @@ void FMassEntityManager::BatchDestroyEntityChunks(TConstArrayView<FMassArchetype
 
 			FMassArchetypeData& ArchetypeData = FMassArchetypeHelper::ArchetypeDataFromHandleChecked(EntityCollection.GetArchetype());
 			ArchetypeData.BatchDestroyEntityChunks(EntityCollection.GetRanges(), EntitiesRemoved);
-		
+
 			GetEntityStorageInterface().Release(EntitiesRemoved);
 		}
 		else
@@ -1299,13 +1319,15 @@ void FMassEntityManager::SwapTagsForEntity(FMassEntityHandle Entity, TNotNull<co
 	checkf(UE::Mass::IsA<FMassTag>(OldTagType), TEXT("%hs works only with tags while '%s' is not one."), __FUNCTION__, *GetPathNameSafe(OldTagType));
 	checkf(UE::Mass::IsA<FMassTag>(NewTagType), TEXT("%hs works only with tags while '%s' is not one."), __FUNCTION__, *GetPathNameSafe(NewTagType));
 
-	if (OldTagType == NewTagType)
+	if (OldTagType == NewTagType
+		|| !ensureMsgf(UE::Mass::IsSparse(OldTagType) == false, TEXT("%hs: %s is a sparse element, not supported"), __FUNCTION__, *OldTagType->GetName())
+		|| !ensureMsgf(UE::Mass::IsSparse(NewTagType) == false, TEXT("%hs: %s is a sparse element, not supported"), __FUNCTION__, *NewTagType->GetName()))
 	{
 		return;
 	}
 
 	CheckIfEntityIsActive(Entity);
-		
+
 	TSharedPtr<FMassArchetypeData>& CurrentArchetype = GetEntityStorageInterface().GetArchetypeAsShared(Entity.Index);
 	check(CurrentArchetype);
 
@@ -1330,22 +1352,29 @@ void FMassEntityManager::AddTagToEntity(FMassEntityHandle Entity, TNotNull<const
 	checkf(UE::Mass::IsA<FMassTag>(TagType), TEXT("%hs works only with tags while '%s' is not one."), __FUNCTION__, *GetPathNameSafe(TagType));
 
 	CheckIfEntityIsActive(Entity);
-	
-	TSharedPtr<FMassArchetypeData>& CurrentArchetype = GetEntityStorageInterface().GetArchetypeAsShared(Entity.Index);
-	check(CurrentArchetype);
 
-	if (CurrentArchetype->HasTagType(TagType) == false)
+	if (UE::Mass::IsSparse(TagType) == false)
 	{
-		FMassElementBitSet NewCompositionBitSet = CurrentArchetype->GetCompositionBitSet();
-		NewCompositionBitSet.Add(*TagType);
-		const FMassArchetypeHandle NewArchetypeHandle = InternalCreateSimilarArchetype(CurrentArchetype, MoveTemp(NewCompositionBitSet));
-		checkSlow(NewArchetypeHandle.IsValid());
+		TSharedPtr<FMassArchetypeData>& CurrentArchetype = GetEntityStorageInterface().GetArchetypeAsShared(Entity.Index);
+		check(CurrentArchetype);
 
-		// Move the entity over
-		CurrentArchetype->MoveEntityToAnotherArchetype(Entity, *NewArchetypeHandle.DataPtr.Get());
-		GetEntityStorageInterface().SetArchetypeFromShared(Entity.Index, NewArchetypeHandle.DataPtr);
+		if (CurrentArchetype->HasTagType(TagType) == false)
+		{
+			FMassElementBitSet NewCompositionBitSet = CurrentArchetype->GetCompositionBitSet();
+			NewCompositionBitSet.Add(*TagType);
+			const FMassArchetypeHandle NewArchetypeHandle = InternalCreateSimilarArchetype(CurrentArchetype, MoveTemp(NewCompositionBitSet));
+			checkSlow(NewArchetypeHandle.IsValid());
+
+			// Move the entity over
+			CurrentArchetype->MoveEntityToAnotherArchetype(Entity, *NewArchetypeHandle.DataPtr.Get());
+			GetEntityStorageInterface().SetArchetypeFromShared(Entity.Index, NewArchetypeHandle.DataPtr);
 		
-		ObserverManager.OnPostCompositionAdded(Entity, FMassArchetypeCompositionDescriptor(FMassTagBitSet(*TagType)));
+			ObserverManager.OnPostCompositionAdded(Entity, FMassArchetypeCompositionDescriptor(FMassTagBitSet(*TagType)));
+		}
+	}
+	else
+	{
+		FStructView _ = AddSparseElementToEntity(Entity, TagType);
 	}
 }
 	
@@ -1356,25 +1385,76 @@ void FMassEntityManager::RemoveTagFromEntity(FMassEntityHandle Entity, TNotNull<
 
 	CheckIfEntityIsActive(Entity);
 
-	TSharedPtr<FMassArchetypeData>& CurrentArchetype = GetEntityStorageInterface().GetArchetypeAsShared(Entity.Index);
-	check(CurrentArchetype);
+	if (UE::Mass::IsSparse(TagType) == false)
+	{
+		TSharedPtr<FMassArchetypeData>& CurrentArchetype = GetEntityStorageInterface().GetArchetypeAsShared(Entity.Index);
+		check(CurrentArchetype);
+
+		if (CurrentArchetype->HasTagType(TagType))
+		{
+			FMassElementBitSet TagDelta(*TagType);
+			FMassElementBitSet NewComposition = CurrentArchetype->GetCompositionBitSet() - TagDelta;
+
+			ObserverManager.OnPreCompositionRemoved(Entity, FMassArchetypeCompositionDescriptor(MoveTemp(TagDelta)));
 
-	if (CurrentArchetype->HasTagType(TagType))
+			const FMassArchetypeHandle NewArchetypeHandle = InternalCreateSimilarArchetype(CurrentArchetype, MoveTemp(NewComposition));
+			checkSlow(NewArchetypeHandle.IsValid());
+
+			// Move the entity over
+			CurrentArchetype->MoveEntityToAnotherArchetype(Entity, *NewArchetypeHandle.DataPtr.Get());
+			GetEntityStorageInterface().SetArchetypeFromShared(Entity.Index, NewArchetypeHandle.DataPtr);
+		}
+	}
+	else
 	{
-		FMassElementBitSet TagDelta(*TagType);
-		FMassElementBitSet NewComposition = CurrentArchetype->GetCompositionBitSet() - TagDelta;
+		RemoveSparseElementFromEntity(Entity, TagType);
+	}
+}
+
+FStructView FMassEntityManager::AddSparseElementToEntity(FMassEntityHandle EntityHandle, TNotNull<const UScriptStruct*> ElementType)
+{
+	CHECK_SPARSE_ELEMENT(ElementType);
 
-		ObserverManager.OnPreCompositionRemoved(Entity, FMassArchetypeCompositionDescriptor(MoveTemp(TagDelta)));
+	CHECK_SYNC_API_RETURN(return {});
+	CheckIfEntityIsActive(EntityHandle);
 
-		const FMassArchetypeHandle NewArchetypeHandle = InternalCreateSimilarArchetype(CurrentArchetype, MoveTemp(NewComposition));
-		checkSlow(NewArchetypeHandle.IsValid());
+	TSharedPtr<FMassArchetypeData>& CurrentArchetype = GetEntityStorageInterface().GetArchetypeAsShared(EntityHandle.Index);
+	check(CurrentArchetype);
 
-		// Move the entity over
-		CurrentArchetype->MoveEntityToAnotherArchetype(Entity, *NewArchetypeHandle.DataPtr.Get());
-		GetEntityStorageInterface().SetArchetypeFromShared(Entity.Index, NewArchetypeHandle.DataPtr);
+	return CurrentArchetype->AddSparseElementToEntity(EntityHandle, ElementType);
+}
+
+void FMassEntityManager::AddSparseElementsToEntity(FMassEntityHandle EntityHandle, const FMassElementBitSet& ElementsBitset)
+{
+	CHECK_SYNC_API_RETURN(return);
+	CheckIfEntityIsActive(EntityHandle);
+
+	TSharedPtr<FMassArchetypeData>& CurrentArchetype = GetEntityStorageInterface().GetArchetypeAsShared(EntityHandle.Index);
+	check(CurrentArchetype);
+
+	for (FMassElementBitSet::FIndexIterator It = ElementsBitset.GetIndexIterator(); It; ++It)
+	{
+		const UScriptStruct* SparseElementType = FMassElementBitSet::GetTypeAtIndex(*It);
+		CA_ASSUME(SparseElementType);
+		// assuming SparseElementType is indeed sparse, as tested by UE::Mass::IsSparse,
+		// AddSparseElementToEntity will ensure that. 
+		FStructView _ = CurrentArchetype->AddSparseElementToEntity(EntityHandle, SparseElementType);
 	}
 }
 
+bool FMassEntityManager::RemoveSparseElementFromEntity(FMassEntityHandle EntityHandle, TNotNull<const UScriptStruct*> ElementType)
+{
+	CHECK_SPARSE_ELEMENT(ElementType);
+
+	CHECK_SYNC_API_RETURN(return false);
+	CheckIfEntityIsActive(EntityHandle);
+
+	TSharedPtr<FMassArchetypeData>& CurrentArchetype = GetEntityStorageInterface().GetArchetypeAsShared(EntityHandle.Index);
+	check(CurrentArchetype);
+
+	return CurrentArchetype->RemoveSparseElementFromEntity(EntityHandle, ElementType);
+}
+
 bool FMassEntityManager::DoesEntityHaveElement(FMassEntityHandle Entity, TNotNull<const UScriptStruct*> ElementType) const
 {
 	CHECK_ELEMENT(ElementType);
@@ -1382,12 +1462,12 @@ bool FMassEntityManager::DoesEntityHaveElement(FMassEntityHandle Entity, TNotNul
 	CheckIfEntityIsActive(Entity);
 	const TSharedPtr<FMassArchetypeData>& CurrentArchetype = GetEntityStorageInterface().GetArchetypeAsShared(Entity.Index);
 	check(CurrentArchetype);
-	if (UE::Mass::IsA<FMassTag>(ElementType))
+	if (UE::Mass::IsSparse(ElementType))
 	{
-		return CurrentArchetype->GetCompositionDescriptor().Contains(ElementType);
+		return CurrentArchetype->HasSparseElementForEntity(ElementType, Entity);
 	}
 	
-	return CurrentArchetype->GetCompositionDescriptor().Contains(ElementType);
+	return CurrentArchetype->GetCompositionBitSet().Contains(*ElementType);
 }
 
 void FMassEntityManager::BatchAddElementToEntities(TConstArrayView<FMassEntityHandle> Entities, TNotNull<const UScriptStruct*> ElementType)
@@ -1397,7 +1477,11 @@ void FMassEntityManager::BatchAddElementToEntities(TConstArrayView<FMassEntityHa
 	TArray<FMassArchetypeEntityCollection> EntityCollections;
 	UE::Mass::Utils::CreateEntityCollections(*this, Entities, FMassArchetypeEntityCollection::EDuplicatesHandling::FoldDuplicates, EntityCollections);
 
-	if (UE::Mass::IsA<FMassTag>(ElementType))
+	if (UE::Mass::IsSparse(ElementType))
+	{
+		BatchAddSparseElementToEntities(EntityCollections, ElementType);
+	}
+	else if (UE::Mass::IsA<FMassTag>(ElementType))
 	{
 		BatchChangeTagsForEntities(EntityCollections, FMassTagBitSet(*ElementType), {});
 	}
@@ -1411,7 +1495,14 @@ void FMassEntityManager::AddElementToEntity(FMassEntityHandle EntityHandle, TNot
 {
 	CHECK_ELEMENT(ElementType);
 
-	if (UE::Mass::IsA<FMassTag>(ElementType))
+	if (UE::Mass::IsSparse(ElementType))
+	{
+		TSharedPtr<FMassArchetypeData>& CurrentArchetype = GetEntityStorageInterface().GetArchetypeAsShared(EntityHandle.Index);
+		checkf(CurrentArchetype, TEXT("Entity (ID: %d, SN:%d) not assigned an archetype yet"), EntityHandle.Index, EntityHandle.SerialNumber);
+
+		FStructView _ = CurrentArchetype->AddSparseElementToEntity(EntityHandle, ElementType);	
+	}
+	else if (UE::Mass::IsA<FMassTag>(ElementType))
 	{
 		AddTagToEntity(EntityHandle, ElementType);
 	}
@@ -1427,8 +1518,11 @@ void FMassEntityManager::BatchRemoveElementFromEntities(TConstArrayView<FMassEnt
 
 	TArray<FMassArchetypeEntityCollection> EntityCollections;
 	UE::Mass::Utils::CreateEntityCollections(*this, Entities, FMassArchetypeEntityCollection::EDuplicatesHandling::FoldDuplicates, EntityCollections);
-
-	if (UE::Mass::IsA<FMassTag>(ElementType))
+	if (UE::Mass::IsSparse(ElementType))
+	{
+		BatchRemoveSparseElementFromEntities(EntityCollections, ElementType);
+	}
+	else if (UE::Mass::IsA<FMassTag>(ElementType))
 	{
 		BatchChangeTagsForEntities(EntityCollections, {}, FMassTagBitSet(*ElementType));
 	}
@@ -1442,7 +1536,11 @@ void FMassEntityManager::RemoveElementFromEntity(FMassEntityHandle Entity, TNotN
 {
 	CHECK_ELEMENT(ElementType);
 
-	if (UE::Mass::IsA<FMassTag>(ElementType))
+	if (UE::Mass::IsSparse(ElementType))
+	{
+		RemoveSparseElementFromEntity(Entity, ElementType);
+	}
+	else if (UE::Mass::IsA<FMassTag>(ElementType))
 	{
 		RemoveTagFromEntity(Entity, ElementType);
 	}
@@ -1653,7 +1751,7 @@ void FMassEntityManager::BatchChangeCompositionForEntities(TConstArrayView<FMass
 		{
 			FMassElementBitSet ElementsAdded;
 			FMassElementBitSet NewComposition = CurrentArchetype->GetCompositionBitSet() + ElementsToAdd - ElementsToRemove;
-
+			check(NewComposition.IsEmpty()==false);
 			if (UNLIKELY(CurrentArchetype->GetCompositionBitSet().IsEquivalent(NewComposition)))
 			{
 				continue;
@@ -1711,7 +1809,12 @@ void FMassEntityManager::BatchAddFragmentInstancesForEntities(TConstArrayView<FM
 
 	TRACE_CPUPROFILER_EVENT_SCOPE(Mass_BatchAddFragmentInstancesForEntities);
 
-	FMassElementBitSet ElementsAffected = FragmentsAffected;
+	const FMassElementBitSet SparseElementsAffected = FragmentsAffected & FMassElementBitSet::GetAllSparseElements();
+	FMassElementBitSet RegularElementsAffected = FragmentsAffected;
+	if (SparseElementsAffected.IsEmpty() == false)
+	{
+		RegularElementsAffected -= SparseElementsAffected;
+	}
 
 	// here's the scenario:
 	// * we get entities from potentially different archetypes
@@ -1723,7 +1826,7 @@ void FMassEntityManager::BatchAddFragmentInstancesForEntities(TConstArrayView<FM
 	// * @todo note that after adding fragment type some entities originally in different archetypes end up in the same 
 	//		archetype. This could be utilized as a basis for optimization. To be investigated.
 	//
-	const bool bFragmentsAddedAreObserved = ObserverManager.HasObserversForBitSet(ElementsAffected, EMassObservedOperation::AddElement);
+	const bool bFragmentsAddedAreObserved = ObserverManager.HasObserversForBitSet(FragmentsAffected, EMassObservedOperation::AddElement);
 
 	for (const FMassArchetypeEntityCollectionWithPayload& EntityRangesWithPayload : EntityCollections)
 	{
@@ -1733,16 +1836,15 @@ void FMassEntityManager::BatchAddFragmentInstancesForEntities(TConstArrayView<FM
 		if (CurrentArchetype)
 		{
 			FMassArchetypeEntityCollection::FEntityRangeArray TargetArchetypeEntityRanges;
-			FMassElementBitSet NewComposition = ElementsAffected + CurrentArchetype->GetCompositionBitSet();
+			FMassElementBitSet NewComposition = RegularElementsAffected + CurrentArchetype->GetCompositionBitSet();
 			FMassElementBitSet ElementsAdded;
 
 			if (CurrentArchetype->GetCompositionBitSet().IsEquivalent(NewComposition) == false)
 			{
 				if (bFragmentsAddedAreObserved)
 				{
-					ElementsAdded = ElementsAffected - CurrentArchetype->GetCompositionBitSet();
+					ElementsAdded = RegularElementsAffected - CurrentArchetype->GetCompositionBitSet();
 				}
-				
 
 				FMassArchetypeHandle NewArchetypeHandle = InternalCreateSimilarArchetype(TargetArchetypeHandle.DataPtr, MoveTemp(NewComposition));
 				checkSlow(NewArchetypeHandle.IsValid());
@@ -1763,6 +1865,7 @@ void FMassEntityManager::BatchAddFragmentInstancesForEntities(TConstArrayView<FM
 			}
 			else
 			{
+				// @todo consider adding sparse elements to ElementsAdded once we start observing those.
 				TargetArchetypeEntityRanges = EntityRangesWithPayload.GetEntityCollection().GetRanges();
 			}
 
@@ -1770,7 +1873,7 @@ void FMassEntityManager::BatchAddFragmentInstancesForEntities(TConstArrayView<FM
 			// note that even though the "subchunk" information could have changed the order of entities is the same and 
 			// corresponds to the order in FMassArchetypeEntityCollectionWithPayload's payload
 			TargetArchetypeHandle.DataPtr->BatchSetFragmentValues(TargetArchetypeEntityRanges, EntityRangesWithPayload.GetPayload());
-			
+
 			if (bFragmentsAddedAreObserved && (ElementsAdded.IsEmpty() == false))
 			{
 				ObserverManager.OnCompositionChanged(
@@ -1823,6 +1926,34 @@ void FMassEntityManager::BatchAddSharedFragmentsForEntities(TConstArrayView<FMas
 	}
 }
 
+void FMassEntityManager::BatchAddSparseElementToEntities(TConstArrayView<FMassArchetypeEntityCollection> EntityCollections, TNotNull<const UScriptStruct*> ElementType)
+{
+	CHECK_SPARSE_ELEMENT(ElementType);
+	CHECK_SYNC_API();
+
+	for (const FMassArchetypeEntityCollection& Collection : EntityCollections)
+	{
+		if (FMassArchetypeData* CurrentArchetype = Collection.GetArchetype().DataPtr.Get())
+		{
+			CurrentArchetype->BatchAddSparseElementToEntities(Collection, ElementType);
+		}
+	}
+}
+
+void FMassEntityManager::BatchRemoveSparseElementFromEntities(TConstArrayView<FMassArchetypeEntityCollection> EntityCollections, TNotNull<const UScriptStruct*> ElementType)
+{
+	CHECK_SPARSE_ELEMENT(ElementType);
+	CHECK_SYNC_API();
+
+	for (const FMassArchetypeEntityCollection& Collection : EntityCollections)
+	{
+		if (FMassArchetypeData* CurrentArchetype = Collection.GetArchetype().DataPtr.Get())
+		{
+			CurrentArchetype->BatchRemoveSparseElementFromEntities(Collection, ElementType);
+		}
+	}
+}
+
 void FMassEntityManager::MoveEntityToAnotherArchetype(FMassEntityHandle Entity, FMassArchetypeHandle NewArchetypeHandle,
 	const FMassArchetypeSharedFragmentValues* SharedFragmentValuesOverride)
 {
@@ -1992,6 +2123,26 @@ void FMassEntityManager::GetMatchingArchetypes(const FMassFragmentRequirements&
 	}
 }
 
+FStructView FMassEntityManager::GetMutableSparseElementDataForEntity(TNotNull<const UScriptStruct*> ElementType, const FMassEntityHandle EntityHandle)
+{
+	check(UE::Mass::IsA<FMassTag>(ElementType) == false);
+	return SparseElementsStorage.GetMutableElementDataForEntity(EntityHandle, ElementType);
+}
+
+FConstStructView FMassEntityManager::GetSparseElementDataForEntity(TNotNull<const UScriptStruct*> ElementType, const FMassEntityHandle EntityHandle) const
+{
+	check(UE::Mass::IsA<FMassTag>(ElementType) == false);
+	return SparseElementsStorage.GetElementDataForEntity(EntityHandle, ElementType);
+}
+
+FMassElementBitSet FMassEntityManager::GetSparseElementsBitSetForEntity(const FMassEntityHandle EntityHandle) const
+{
+	FMassArchetypeHandle ArchetypeHandle = GetArchetypeForEntity(EntityHandle);
+	return ArchetypeHandle.IsValid()
+		? ArchetypeHandle.DataPtr->GetSparseElementsBitSetForEntity(EntityHandle)
+		: FMassElementBitSet();
+}
+
 FMassExecutionContext FMassEntityManager::CreateExecutionContext(const float DeltaSeconds)
 {
 	FMassExecutionContext ExecutionContext(*this, DeltaSeconds);
@@ -2301,6 +2452,11 @@ bool FMassEntityManager::DebugHasCommandsToFlush() const
 	checkfSlow(NumCommandBuffers == 2, TEXT("This check relies on there being two command buffers."));
 	return DeferredCommandBuffers[0]->HasPendingCommands() || DeferredCommandBuffers[1]->HasPendingCommands();
 }
+
+UE::Mass::FSparseElementsStorage& FMassEntityManager::DebugGetSparseElementsStorage()
+{
+	return SparseElementsStorage;
+}
 #endif // WITH_MASSENTITY_DEBUG
 
 #undef CHECK_SYNC_API
diff --git a/Engine/Source/Runtime/MassEntity/Private/MassEntityView.cpp b/Engine/Source/Runtime/MassEntity/Private/MassEntityView.cpp
--- a/Engine/Source/Runtime/MassEntity/Private/MassEntityView.cpp
+++ b/Engine/Source/Runtime/MassEntity/Private/MassEntityView.cpp
@@ -34,6 +34,8 @@ FMassEntityView FMassEntityView::TryMakeView(const FMassEntityManager& EntityMan
 
 void* FMassEntityView::GetFragmentPtr(const UScriptStruct& FragmentType) const
 {
+	checkSlow(EntityDataHandle.IsValid(Archetype));
+
 	if (testableEnsureMsgf(Archetype, TEXT("%hs: Trying to access fragments while no archetype set"), __FUNCTION__))
 	{
 		CA_ASSUME(Archetype);
@@ -48,6 +50,8 @@ void* FMassEntityView::GetFragmentPtr(const UScriptStruct& FragmentType) const
 
 void* FMassEntityView::GetFragmentPtrChecked(const UScriptStruct& FragmentType) const
 {
+	checkSlow(EntityDataHandle.IsValid(Archetype));
+
 	if (testableEnsureMsgf(Archetype, TEXT("%hs: Trying to access fragments while no archetype set"), __FUNCTION__))
 	{
 		CA_ASSUME(Archetype);
@@ -57,6 +61,30 @@ void* FMassEntityView::GetFragmentPtrChecked(const UScriptStruct& FragmentType)
 	return nullptr;
 }
 
+FStructView FMassEntityView::GetSparseFragmentPtr(const TNotNull<const UScriptStruct*> FragmentType)
+{
+	checkSlow(EntityDataHandle.IsValid(Archetype));
+
+	if (testableEnsureMsgf(Archetype, TEXT("%hs: Trying to access sparse fragments while no archetype set"), __FUNCTION__))
+	{
+		CA_ASSUME(Archetype);
+		return Archetype->GetMutableSparseElementDataForEntity(FragmentType, Entity);
+	}
+	return {};
+}
+
+FConstStructView FMassEntityView::GetSparseFragmentPtr(const TNotNull<const UScriptStruct*> FragmentType) const
+{
+	checkSlow(EntityDataHandle.IsValid(Archetype));
+
+	if (testableEnsureMsgf(Archetype, TEXT("%hs: Trying to access sparse fragments while no archetype set"), __FUNCTION__))
+	{
+		CA_ASSUME(Archetype);
+		return Archetype->GetSparseElementDataForEntity(FragmentType, Entity);
+	}
+	return {};
+}
+
 const void* FMassEntityView::GetConstSharedFragmentPtr(const UScriptStruct& FragmentType) const
 {
 	const FConstSharedStruct* SharedFragment = Archetype->GetSharedFragmentValues(Entity).GetConstSharedFragments().FindByPredicate(FStructTypeEqualOperator(&FragmentType));
@@ -83,8 +111,20 @@ void* FMassEntityView::GetSharedFragmentPtrChecked(const UScriptStruct& Fragment
 	return SharedFragment->GetMemory();
 }
 
-bool FMassEntityView::HasTag(const UScriptStruct& TagType) const
+bool FMassEntityView::HasTag(TNotNull<const UScriptStruct*> TagType, const bool bIncludeSparse) const
+{
+	checkSlow(EntityDataHandle.IsValid(Archetype));
+	return Archetype && (Archetype->HasTagType(TagType) || (bIncludeSparse && Archetype->HasSparseElementForEntity(TagType, EntityDataHandle)));
+}
+
+bool FMassEntityView::HasFragment(TNotNull<const UScriptStruct*> FragmentType, const bool bIncludeSparse) const
+{
+	checkSlow(EntityDataHandle.IsValid(Archetype));
+	return Archetype && (Archetype->HasFragmentType(FragmentType) || (bIncludeSparse && Archetype->HasSparseElementForEntity(FragmentType, EntityDataHandle)));
+}
+
+bool FMassEntityView::HasElement(TNotNull<const UScriptStruct*> ElementType, const bool bIncludeSparse) const
 {
-	check(EntityDataHandle.IsValid(Archetype));
-	return Archetype->HasTagType(&TagType);
+	checkSlow(EntityDataHandle.IsValid(Archetype));
+	return Archetype && (Archetype->HasElement(ElementType) || (bIncludeSparse && Archetype->HasSparseElementForEntity(ElementType, EntityDataHandle)));
 }
diff --git a/Engine/Source/Runtime/MassEntity/Private/MassExecutionContext.cpp b/Engine/Source/Runtime/MassEntity/Private/MassExecutionContext.cpp
--- a/Engine/Source/Runtime/MassEntity/Private/MassExecutionContext.cpp
+++ b/Engine/Source/Runtime/MassEntity/Private/MassExecutionContext.cpp
@@ -56,14 +56,23 @@ void FMassExecutionContext::ClearExecutionData()
 	ChunkFragmentViews.Reset();
 	ConstSharedFragmentViews.Reset();
 	SharedFragmentViews.Reset();
+	RequestedSparseElementsAccess.Reset();
 	CurrentArchetypeCompositionBitSet.Reset();
 	EntityListView = {};
-	ChunkSerialModificationNumber = -1;
+	ChunkEntityIndex = {};
+	ArchetypeDataPtr = nullptr;
+	bHasSparseRequirements = false;
 #if WITH_MASSENTITY_DEBUG
 	DebugColor = FColor();
 #endif // WITH_MASSENTITY_DEBUG
 }
 
+void FMassExecutionContext::SetCurrentArchetype(FMassArchetypeData& Archetype)
+{
+	CurrentArchetypeCompositionBitSet = Archetype.GetCompositionBitSet();
+	ArchetypeDataPtr = &Archetype;
+}
+
 bool FMassExecutionContext::CacheSubsystemRequirements(const FMassSubsystemRequirements& SubsystemRequirements)
 {
 	return SubsystemAccess.CacheSubsystemRequirements(SubsystemRequirements);
@@ -119,6 +128,15 @@ void FMassExecutionContext::SetFragmentRequirements(const FMassFragmentRequireme
 			SharedFragmentViews.Emplace(Requirement);
 		}
 	}
+
+	RequestedSparseElementsAccess.Reset();
+	bHasSparseRequirements = FragmentRequirements.HasSparseRequirements();
+	if (bHasSparseRequirements)
+	{
+		RequestedSparseElementsAccess = FragmentRequirements.GetRequiredAllSparseElements()
+			+ FragmentRequirements.GetRequiredAnySparseElements()
+			+ FragmentRequirements.GetRequiredOptionalSparseElements();
+	}
 }
 
 UWorld* FMassExecutionContext::GetWorld() const
@@ -227,6 +245,14 @@ FMassExecutionContext::FEntityIterator FMassExecutionContext::CreateEntityIterat
 	return FEntityIterator(*this, QueriesStack.Last());
 }
 
+FMassExecutionContext::FSparseEntityIterator FMassExecutionContext::CreateSparseEntityIterator()
+{
+	testableCheckfReturn(QueriesStack.Num()
+		, return {}
+		, TEXT("Attempting to create a Sparse Entity Iterator when no entity query is being executed."));
+	return FSparseEntityIterator(FEntityIterator(*this, QueriesStack.Last()));
+}
+
 FMassEntityHandle FMassExecutionContext::CreateEntity(const FMassArchetypeHandle& ArchetypeHandle,
 	const FMassArchetypeSharedFragmentValues& SharedValues)
 {
@@ -269,6 +295,49 @@ FMassExecutionContext& FMassExecutionContext::GetDummyInstance()
 	return DummyContext;
 }
 
+bool FMassExecutionContext::SetCurrentChunk(const int32 ChunkIndex, const FMassArchetypeChunk& Chunk)
+{
+	if (bHasSparseRequirements == false || ExecuteSparseChunkFilterImpl(Chunk))
+	{
+		ChunkEntityIndex = FMassEntityInChunkDataHandle(Chunk.GetRawMemory(), 0, ChunkIndex, Chunk.GetSerialModificationNumber());
+		return true;
+	}
+	ChunkEntityIndex = {};
+	return false;
+}
+
+FStructView FMassExecutionContext::GetSparseElementInternal(TNotNull<const UScriptStruct*> ElementType, const FMassEntityHandle EntityHandle) const
+{
+	return EntityManager->GetMutableSparseElementDataForEntity(ElementType, EntityHandle);
+}
+
+bool FMassExecutionContext::HasSparseElement(TNotNull<const UScriptStruct*> ElementType, const FEntityIterator& EntityIterator) const
+{
+	return ArchetypeDataPtr && ArchetypeDataPtr->HasSparseElementForEntity(ElementType, EntityListView[EntityIterator]);
+}
+
+bool FMassExecutionContext::DoesChunkContainSparseElement(TNotNull<const UScriptStruct*> ElementType) const
+{
+	// we're not checking if ChunkEntityIndex.ChunkIndex is valid, ArchetypeDataPtr->DoesChunkContainSparseElement will do that.
+	return ArchetypeDataPtr && ArchetypeDataPtr->DoesChunkContainSparseElement(ElementType, ChunkEntityIndex.ChunkIndex);
+}
+
+bool FMassExecutionContext::ExecuteSparseElementsFilterImpl(const FSparseEntityIterator& EntityIt) const
+{
+	// @todo we could gain some perf if we cached required sparse elements bitsets as member variables. Needs perf measuring if worth it.
+	checkSlow(ArchetypeDataPtr);
+	const FMassEntityInChunkDataHandle EntityInChunkHandle(ChunkEntityIndex, EntityIt);
+	return ArchetypeDataPtr->DoesEntityMatchSparseComposition(EntityInChunkHandle
+		, QueriesStack.Last().Query->GetRequiredAllSparseElements()
+		, QueriesStack.Last().Query->GetRequiredAnySparseElements()
+		, QueriesStack.Last().Query->GetRequiredNoneSparseElements());
+}
+
+inline bool FMassExecutionContext::ExecuteSparseChunkFilterImpl(const FMassArchetypeChunk& Chunk) const
+{
+	return Chunk.DoesMatchComposition(QueriesStack.Last().Query->GetRequiredAllSparseElements(), QueriesStack.Last().Query->GetRequiredAnySparseElements());
+}
+
 //-----------------------------------------------------------------------------
 // FMassExecutionContext::FQueryTransientRuntime
 //-----------------------------------------------------------------------------
diff --git a/Engine/Source/Runtime/MassEntity/Private/MassObserverRegistry.cpp b/Engine/Source/Runtime/MassEntity/Private/MassObserverRegistry.cpp
--- a/Engine/Source/Runtime/MassEntity/Private/MassObserverRegistry.cpp
+++ b/Engine/Source/Runtime/MassEntity/Private/MassObserverRegistry.cpp
@@ -22,6 +22,8 @@ UMassObserverRegistry::UMassObserverRegistry()
 void UMassObserverRegistry::RegisterObserver(TNotNull<const UScriptStruct*> ObservedType, const uint8 OperationFlags, const TSubclassOf<UMassProcessor>& ObserverClass)
 {
 	check(ObserverClass);
+	ensureAlwaysMsgf(UE::Mass::IsSparse(ObservedType) == false, TEXT("%hs: trying to register observer for %s, but observing sparse elements is not supported at the moment.")
+		, __FUNCTION__, *ObservedType->GetName());
 
 	for (uint8 OperationIndex = 0; OperationIndex < static_cast<uint8>(EMassObservedOperation::MAX); ++OperationIndex)
 	{
@@ -34,6 +36,9 @@ void UMassObserverRegistry::RegisterObserver(TNotNull<const UScriptStruct*> Obse
 
 void UMassObserverRegistry::RegisterObserver(const UScriptStruct& ObservedType, const EMassObservedOperation Operation, const TSubclassOf<UMassProcessor>& ObserverClass)
 {
+	ensureAlwaysMsgf(UE::Mass::IsSparse(&ObservedType) == false, TEXT("%hs: trying to register observer for %s, but observing sparse elements is not supported at the moment.")
+		, __FUNCTION__, *ObservedType.GetName());
+
 	RegisterObserver(&ObservedType, static_cast<uint8>(1 << static_cast<uint8>(Operation)), ObserverClass);
 }
 
diff --git a/Engine/Source/Runtime/MassEntity/Private/MassRequirements.cpp b/Engine/Source/Runtime/MassEntity/Private/MassRequirements.cpp
--- a/Engine/Source/Runtime/MassEntity/Private/MassRequirements.cpp
+++ b/Engine/Source/Runtime/MassEntity/Private/MassRequirements.cpp
@@ -162,6 +162,10 @@ uint32 GetTypeHash(const FMassFragmentRequirements& Instance)
 	Hash = HashCombine(Hash, GetTypeHash(Instance.RequiredAnyFragments));
 	Hash = HashCombine(Hash, GetTypeHash(Instance.RequiredOptionalFragments));
 	Hash = HashCombine(Hash, GetTypeHash(Instance.RequiredNoneFragments));
+	Hash = HashCombine(Hash, GetTypeHash(Instance.RequiredAllSparseElements));
+	Hash = HashCombine(Hash, GetTypeHash(Instance.RequiredAnySparseElements));
+	Hash = HashCombine(Hash, GetTypeHash(Instance.RequiredOptionalSparseElements));
+	Hash = HashCombine(Hash, GetTypeHash(Instance.RequiredNoneSparseElements));
 	Hash = HashCombine(Hash, GetTypeHash(Instance.RequiredAllChunkFragments));
 	Hash = HashCombine(Hash, GetTypeHash(Instance.RequiredOptionalChunkFragments));
 	Hash = HashCombine(Hash, GetTypeHash(Instance.RequiredNoneChunkFragments));
@@ -222,11 +226,18 @@ FORCEINLINE void FMassFragmentRequirements::CacheProperties() const
 			&& RequiredNoneSharedFragments.IsEmpty()
 			&& RequiredNoneConstSharedFragments.IsEmpty());
 
+		// note that we're not including RequiredOptionalSparseElements on purpose here
+		// as that's just a way to signal that a given sparse element might get accessed
+		// and does not affect archetype filtering
 		bHasOptionalRequirements = !(RequiredOptionalFragments.IsEmpty()
-				&& RequiredOptionalTags.IsEmpty()
-				&& RequiredOptionalChunkFragments.IsEmpty()
-				&& RequiredOptionalSharedFragments.IsEmpty()
-				&& RequiredOptionalConstSharedFragments.IsEmpty());
+			&& RequiredOptionalTags.IsEmpty()
+			&& RequiredOptionalChunkFragments.IsEmpty()
+			&& RequiredOptionalSharedFragments.IsEmpty()
+			&& RequiredOptionalConstSharedFragments.IsEmpty());
+
+		bHasSparseRequirements = !(RequiredAllSparseElements.IsEmpty() 
+			&& RequiredAnySparseElements.IsEmpty()
+			&& RequiredNoneSparseElements.IsEmpty());
 
 		bPropertiesCached = true;
 	}
@@ -236,15 +247,14 @@ bool FMassFragmentRequirements::CheckValidity() const
 {
 	CacheProperties();
 	// @todo we need to add more sophisticated testing somewhere to detect contradicting requirements - like having and not having a given tag.
+	// note: bHasSparseRequirements not being tested on purpose - sparse-only requirements are not valid
 	return bHasPositiveRequirements || bHasNegativeRequirements || bHasOptionalRequirements;
 }
 
 bool FMassFragmentRequirements::IsEmpty() const
 {
 	CacheProperties();
-	// note that even though at the moment the following condition is the same as negation of current CheckValidity value
-	// that will change in the future (with additional validity checks).
-	return !bHasPositiveRequirements && !bHasNegativeRequirements && !bHasOptionalRequirements;
+	return !bHasPositiveRequirements && !bHasNegativeRequirements && !bHasOptionalRequirements && !bHasSparseRequirements;
 }
 
 bool FMassFragmentRequirements::DoesMatchAnyOptionals(const FMassElementBitSet& ArchetypeCompositionBitSet) const
@@ -316,10 +326,15 @@ void FMassFragmentRequirements::Reset()
 	RequiredAllTags.Reset();
 	RequiredAnyTags.Reset();
 	RequiredNoneTags.Reset();
+	RequiredOptionalTags.Reset();
 	RequiredAllFragments.Reset();
 	RequiredAnyFragments.Reset();
 	RequiredOptionalFragments.Reset();
 	RequiredNoneFragments.Reset();
+	RequiredAllSparseElements.Reset();
+	RequiredAnySparseElements.Reset();
+	RequiredOptionalSparseElements.Reset();
+	RequiredNoneSparseElements.Reset();
 	RequiredAllChunkFragments.Reset();
 	RequiredOptionalChunkFragments.Reset();
 	RequiredNoneChunkFragments.Reset();
@@ -332,6 +347,8 @@ void FMassFragmentRequirements::Reset()
 
 	IncrementalChangesCount = 0;
 
+	bPropertiesCached = false;
+
 	// note that we're not resetting bInitialized nor CachedEntityManager, on purpose
 	// the point of this function is to just reset the contents while still being able
 	// to add elements to it. This "requirements" instance is now "empty" but still valid
diff --git a/Engine/Source/Runtime/MassEntity/Private/MassSparseElementsStorage.cpp b/Engine/Source/Runtime/MassEntity/Private/MassSparseElementsStorage.cpp
new file mode 100644
--- /dev/null
+++ b/Engine/Source/Runtime/MassEntity/Private/MassSparseElementsStorage.cpp
@@ -0,0 +1,364 @@
+// Copyright Epic Games, Inc. All Rights Reserved.
+
+#include "MassSparseElementsStorage.h"
+#include "MassEntityManager.h"
+
+namespace UE::Mass
+{
+inline void SanitizeEntityHandleArray(TArray<FMassEntityHandle>& InOutEntityHandles)
+{
+	// sorting because data is stored in ascending Index order - EntityHandle.Index is used as the index to TypePool's elements
+	InOutEntityHandles.Sort();
+	if (InOutEntityHandles.IsEmpty() == false && InOutEntityHandles[0].IsSet() == false)
+	{
+		InOutEntityHandles.RemoveAll([](const FMassEntityHandle& Handle)
+			{
+				return Handle.IsSet() == false;
+			});
+	}
+}
+
+//-----------------------------------------------------------------------------
+// FSparseElementsStorage::FTypeConfig
+//-----------------------------------------------------------------------------
+FSparseElementsStorage::FTypeConfig::FTypeConfig()
+	: ElementType(FMassElement::StaticStruct())
+{
+	// this is a dummy constructor, used by FTypePool::DummyPool
+}
+
+//-----------------------------------------------------------------------------
+// FSparseElementsStorage::FTypePool
+//-----------------------------------------------------------------------------
+FSparseElementsStorage::FTypePool::~FTypePool()
+{
+	for (FPackedStructDataChunk& PackedData : Chunks)
+	{
+		if (PackedData.RawMemory && PackedData.NumElements > 0 && PackedData.OccupationMask.IsEmpty() == false)
+		{
+			for (TBitArray<>::FConstIterator It(PackedData.OccupationMask); It; ++It)
+			{
+				if (It.GetValue())
+				{
+					Config.GetType()->DestroyStruct(PackedData.GetElement(It.GetIndex(), Config.GetElementSize()));
+				}
+			}
+		}
+	}
+}
+
+TNotNull<uint8*> FSparseElementsStorage::FTypePool::Add(const uint32 EntityIndex)
+{
+	const uint32 ChunkIndex = GetChunkIndex(EntityIndex);
+	const uint32 ElementIndex = GetElementIndex(EntityIndex);
+
+	FPackedStructDataChunk& PackedData = GetOrCreatePackedData(ChunkIndex);
+	if (PackedData.HasElement(ElementIndex) == false)
+	{
+		uint8* Memory = PackedData.Add(ElementIndex, Config.GetElementSize());
+		Config.GetType()->InitializeDefaultValue(Memory);
+		++NumElements; 
+		return Memory;
+	}
+	
+	return PackedData.GetElementUnsafe(ElementIndex, Config.GetElementSize());
+}
+
+TNotNull<uint8*> FSparseElementsStorage::FTypePool::Add(const uint32 EntityIndex, TNotNull<const uint8*> SourceData)
+{
+	TNotNull<uint8*> ElementPtr = Add(EntityIndex);
+	Config.GetType()->CopyScriptStruct(ElementPtr, SourceData);
+	return ElementPtr;
+}
+
+FStructView FSparseElementsStorage::FTypePool::Add_GetView(const uint32 EntityIndex)
+{	
+	return { Config.GetType(), Add(EntityIndex) };
+}
+
+bool FSparseElementsStorage::FTypePool::Remove(const uint32 EntityIndex)
+{
+	const uint32 ChunkIndex = GetChunkIndex(EntityIndex);
+	const uint32 ElementIndex = GetElementIndex(EntityIndex);
+
+	if (FPackedStructDataChunk* PackedData = GetPackedData(ChunkIndex))
+	{
+		if (uint8* Memory = PackedData->Remove(ElementIndex, Config.GetElementSize()))
+		{
+			Config.GetType()->DestroyStruct(Memory);
+			FMemory::Memzero(Memory, Config.GetElementSize());
+			--NumElements; 
+			return true;
+		}
+	}
+	return false;
+}
+
+FConstStructView FSparseElementsStorage::FTypePool::Get(const uint32 EntityIndex) const
+{
+	return FConstStructView(const_cast<FTypePool*>(this)->GetMutable(EntityIndex));
+}
+
+FStructView FSparseElementsStorage::FTypePool::GetMutable(const uint32 EntityIndex)
+{
+	const uint32 ChunkIndex = GetChunkIndex(EntityIndex);
+	const uint32 ElementIndex = GetElementIndex(EntityIndex);
+
+	if (FPackedStructDataChunk* PackedData = GetPackedData(ChunkIndex))
+	{
+		uint8* Memory = PackedData->GetElement(ElementIndex, Config.GetElementSize());
+		return { Config.GetType(), Memory };
+	}
+	return {};
+}
+
+FStructView FSparseElementsStorage::FTypePool::GetOrCreateMutable(const uint32 EntityIndex)
+{
+	const uint32 ChunkIndex = GetChunkIndex(EntityIndex);
+	const uint32 ElementIndex = GetElementIndex(EntityIndex);
+
+	FPackedStructDataChunk& PackedData = GetOrCreatePackedData(ChunkIndex);
+	// note that it's fine for Memory to be nullptr. The return value should always be checked for validity.
+	uint8* Memory = PackedData.GetElement(ElementIndex, Config.GetElementSize());
+	return {Config.GetType(), Memory};
+}
+
+FSparseElementsStorage::FPackedStructDataChunk* FSparseElementsStorage::FTypePool::GetPackedData(const uint32 ChunkIndex)
+{
+	return Chunks.IsValidIndex(ChunkIndex) ? &Chunks[ChunkIndex] : nullptr;
+}
+
+const FSparseElementsStorage::FPackedStructDataChunk* FSparseElementsStorage::FTypePool::GetPackedData(const uint32 ChunkIndex) const
+{
+	return Chunks.IsValidIndex(ChunkIndex) ? &Chunks[ChunkIndex] : nullptr;
+}
+
+FSparseElementsStorage::FPackedStructDataChunk& FSparseElementsStorage::FTypePool::GetOrCreatePackedData(const uint32 ChunkIndex)
+{
+	if (Chunks.IsValidIndex(ChunkIndex) == false)
+	{
+		Chunks.AddDefaulted(ChunkIndex - Chunks.Num() + 1);
+	}
+
+	FPackedStructDataChunk& PackedData = Chunks[ChunkIndex];
+
+	if (PackedData.IsInitialized() == false)
+	{
+		PackedData.Init(Config.ElementsPerChunk, Config.GetElementSize());
+	}
+
+	return PackedData;
+}
+
+FSparseElementsStorage::FTypePool& FSparseElementsStorage::FTypePool::GetDummyPool()
+{
+	static FTypePool DummyPool;
+	return DummyPool;
+}
+
+//-----------------------------------------------------------------------------
+// FSparseElementsStorage::FSparseElementIterator 
+//-----------------------------------------------------------------------------
+FMassEntityHandle FSparseElementIterator::GetEntityHandle(const FMassEntityManager& EntityManager) const
+{
+	return EntityManager.CreateEntityIndexHandle(GetEntityIndex());
+}
+
+//-----------------------------------------------------------------------------
+// FSparseElementsStorage
+//-----------------------------------------------------------------------------
+void FSparseElementsStorage::Deinitialize()
+{
+	TypePools.Reset();
+}
+
+FSparseElementsStorage::FTypePool& FSparseElementsStorage::GetOrCreateTypePool(TNotNull<const UScriptStruct*> ElementType, const int32 InTypeIndex)
+{
+	const int32 TypeIndex = InTypeIndex == INDEX_NONE ? FMassElementBitSet::GetTypeIndex(*ElementType) : InTypeIndex;
+
+	if (TypePools.IsValidIndex(TypeIndex) == false)
+	{
+		TypePools.EmplaceAt(TypeIndex, FTypeConfig(ElementType));
+	}
+
+	return TypePools[TypeIndex];
+}
+
+const FSparseElementsStorage::FTypePool* FSparseElementsStorage::GetTypePool(TNotNull<const UScriptStruct*> ElementType) const
+{
+	const int32 TypeIndex = FMassElementBitSet::GetTypeIndex(*ElementType);
+	return TypePools.IsValidIndex(TypeIndex) ? &TypePools[TypeIndex] : nullptr;
+}
+
+FSparseElementsStorage::FTypePool* FSparseElementsStorage::GetTypePool(TNotNull<const UScriptStruct*> ElementType)
+{
+	const int32 TypeIndex = FMassElementBitSet::GetTypeIndex(*ElementType);
+	return TypePools.IsValidIndex(TypeIndex) ? &TypePools[TypeIndex] : nullptr;
+}
+
+FStructView FSparseElementsStorage::AddElementToEntity(FMassEntityHandle EntityHandle, TNotNull<const UScriptStruct*> ElementType)
+{
+	checkSlow(!UE::Mass::IsA<FMassTag>(ElementType));
+	const int32 TypeIndex = FMassElementBitSet::GetTypeIndex(*ElementType);
+	FTypePool& TypePool = GetOrCreateTypePool(ElementType, TypeIndex);
+	return TypePool.Add_GetView(EntityHandle.Index);
+}
+
+bool FSparseElementsStorage::RemoveElementFromEntity(FMassEntityHandle EntityHandle, TNotNull<const UScriptStruct*> ElementType)
+{
+	FTypePool* TypePool = GetTypePool(ElementType);
+	return TypePool && TypePool->Remove(EntityHandle.Index);
+}
+
+template<typename TContainer>
+void FSparseElementsStorage::BatchAddElementToEntities(TContainer InEntityHandles, TNotNull<const UScriptStruct*> ElementType, TArray<FStructView>* OutAddedElements)
+{
+	const int32 TypeIndex = FMassElementBitSet::GetTypeIndex(*ElementType);
+
+	FTypePool& TypePool = GetOrCreateTypePool(ElementType, TypeIndex);
+	// @todo TypePool.Add can be optimized by checking for index continuity and whether the following index is in the same chunk.
+	if (OutAddedElements)
+	{
+		for (const FMassEntityHandle& EntityHandle : InEntityHandles)
+		{
+			OutAddedElements->Add(TypePool.Add_GetView(EntityHandle.Index));
+		}
+	}
+	else
+	{
+		for (const FMassEntityHandle& EntityHandle : InEntityHandles)
+		{
+			TypePool.Add(EntityHandle.Index);
+		}
+	}
+}
+
+void FSparseElementsStorage::BatchAddElementToEntities(TArray<FMassEntityHandle>& InOutEntityHandles, TNotNull<const UScriptStruct*> ElementType, TArray<FStructView>* OutAddedElements)
+{
+	if (!ensureMsgf(UE::Mass::IsA<FMassElement>(ElementType) && !UE::Mass::IsA<FMassTag>(ElementType)
+		, TEXT("%hs: %s is not an element type"), __FUNCTION__, *ElementType->GetName()))
+	{
+		return;
+	}
+	if (InOutEntityHandles.IsEmpty())
+	{
+		return;
+	}
+
+	SanitizeEntityHandleArray(InOutEntityHandles);
+	BatchAddElementToEntities<TArray<FMassEntityHandle>&>(InOutEntityHandles, ElementType, OutAddedElements);
+}
+
+void FSparseElementsStorage::BatchAddElementToEntities(TArrayView<const FMassEntityHandle> InEntityHandles, TNotNull<const UScriptStruct*> ElementType, TArray<FStructView>* OutAddedElements)
+{
+	if (!ensureMsgf(UE::Mass::IsA<FMassElement>(ElementType) && !UE::Mass::IsA<FMassTag>(ElementType)
+		, TEXT("%hs: %s is not an element type"), __FUNCTION__, *ElementType->GetName()))
+	{
+		return;
+	}
+	if (InEntityHandles.IsEmpty())
+	{
+		return;
+	}
+
+	BatchAddElementToEntities<TArrayView<const FMassEntityHandle>>(InEntityHandles, ElementType, OutAddedElements);
+}
+
+void FSparseElementsStorage::BatchAddElementInstancesToEntities(TArrayView<const FMassEntityHandle> InEntityHandles, const FStructArrayView& FragmentPayload)
+{
+	checkf(InEntityHandles.Num() == FragmentPayload.Num(), TEXT("%hs: Payload size is expected to match number of entity handles"), __FUNCTION__);
+	const UScriptStruct* ElementType = FragmentPayload.GetScriptStruct();
+	CA_ASSUME(ElementType);
+	const int32 TypeIndex = FMassElementBitSet::GetTypeIndex(*ElementType);
+
+	FTypePool& TypePool = GetOrCreateTypePool(ElementType, TypeIndex);
+	for (int32 Index = 0; Index < InEntityHandles.Num(); ++Index)
+	{
+		TypePool.Add(InEntityHandles[Index].Index, static_cast<uint8*>(FragmentPayload.GetDataAt(Index)));
+	}
+}
+
+void FSparseElementsStorage::BatchAddElementInstancesToEntities(TArrayView<const FMassEntityHandle> InEntityHandles, const FConstStructView FragmentInstance)
+{
+	check(FragmentInstance.IsValid());
+	const UScriptStruct* ElementType = FragmentInstance.GetScriptStruct();
+	CA_ASSUME(ElementType);
+	const int32 TypeIndex = FMassElementBitSet::GetTypeIndex(*ElementType);
+
+	FTypePool& TypePool = GetOrCreateTypePool(ElementType, TypeIndex);
+	for (int32 Index = 0; Index < InEntityHandles.Num(); ++Index)
+	{
+		TypePool.Add(InEntityHandles[Index].Index, FragmentInstance.GetMemory());
+	}
+}	
+
+int32 FSparseElementsStorage::BatchRemoveElementFromEntities(TArray<FMassEntityHandle>& InOutEntityHandles, TNotNull<const UScriptStruct*> ElementType)
+{
+	SanitizeEntityHandleArray(InOutEntityHandles);
+	return BatchRemoveElementFromEntities(MakeArrayView(InOutEntityHandles), ElementType);
+}
+
+int32 FSparseElementsStorage::BatchRemoveElementFromEntities(TArrayView<const FMassEntityHandle> InEntityHandles, TNotNull<const UScriptStruct*> ElementType)
+{
+	if (InEntityHandles.IsEmpty())
+	{
+		return 0;
+	}
+
+	int32 TotalRemoved = 0;
+
+	if (FTypePool* TypePool = GetTypePool(ElementType))
+	{
+		// @todo TypePool.Remove can be optimized by checking for index continuity and whether the following index is in the same chunk.
+		for (const FMassEntityHandle& EntityHandle : InEntityHandles)
+		{
+			TotalRemoved += static_cast<int32>(TypePool->Remove(EntityHandle.Index));
+		}
+	}
+	return TotalRemoved;
+}
+
+FStructView FSparseElementsStorage::GetMutableElementDataForEntity(const FMassEntityHandle EntityHandle, TNotNull<const UScriptStruct*> ElementType)
+{
+	checkSlow(!UE::Mass::IsA<FMassTag>(ElementType));
+	if (FTypePool* TypePool = GetTypePool(ElementType))
+	{
+		return TypePool->GetMutable(EntityHandle.Index);
+	}
+	return {};
+}
+
+FConstStructView FSparseElementsStorage::GetElementDataForEntity(const FMassEntityHandle EntityHandle, TNotNull<const UScriptStruct*> ElementType) const
+{
+	checkSlow(!UE::Mass::IsA<FMassTag>(ElementType));
+	if (const FTypePool* TypePool = GetTypePool(ElementType))
+	{
+		return TypePool->Get(EntityHandle.Index);
+	}
+	return {};
+}
+
+void FSparseElementsStorage::RemoveEntity(FMassEntityHandle EntityHandle, const FMassElementBitSet& Elements)
+{
+	for (FMassElementBitSet::FIndexIterator It = Elements.GetIndexIterator(); It; ++It)
+	{
+		if (TypePools.IsValidIndex(*It))
+		{
+			TypePools[*It].Remove(EntityHandle.Index);
+		}
+	}
+}
+
+void FSparseElementsStorage::ConfigureType(TNotNull<const UScriptStruct*> ElementType, FTypeConfig& Config)
+{
+	const int32 TypeIndex = FMassElementBitSet::GetTypeIndex(*ElementType);
+
+	if (ensureMsgf(TypePools.IsValidIndex(TypeIndex) == false || TypePools[TypeIndex].GetNumElements() == 0
+		, TEXT("%hs: Configuring types after elements have been instantiated is not supported."), __FUNCTION__))
+	{
+		TypePools.EmplaceAt(TypeIndex, Config);
+	}
+}
+
+
+} // namespace UE::Mass
diff --git a/Engine/Source/Runtime/MassEntity/Public/MassArchetypeTypes.h b/Engine/Source/Runtime/MassEntity/Public/MassArchetypeTypes.h
--- a/Engine/Source/Runtime/MassEntity/Public/MassArchetypeTypes.h
+++ b/Engine/Source/Runtime/MassEntity/Public/MassArchetypeTypes.h
@@ -313,6 +313,7 @@ struct FMassEntityInChunkDataHandle : FMassRawEntityInChunkData
 	FMassEntityInChunkDataHandle& operator=(const FMassEntityInChunkDataHandle&);
 	FMassEntityInChunkDataHandle& operator=(FMassEntityInChunkDataHandle&&);
 	FMassEntityInChunkDataHandle(uint8* InChunkRawMemory, const int32 InIndexWithinChunk, const int32 InChunkIndex, const int32 InChunkSerialNumber);
+	FMassEntityInChunkDataHandle(const FMassEntityInChunkDataHandle& Source, const int32 InIndexWithinChunk);
 
 	MASSENTITY_API bool IsValid(const FMassArchetypeData* ArchetypeData) const;
 	MASSENTITY_API bool IsValid(const FMassArchetypeHandle& ArchetypeHandle) const;
@@ -470,6 +471,12 @@ inline FMassEntityInChunkDataHandle::FMassEntityInChunkDataHandle(uint8* InChunk
 {
 }
 
+inline FMassEntityInChunkDataHandle::FMassEntityInChunkDataHandle(const FMassEntityInChunkDataHandle& Source, const int32 InIndexWithinChunk)
+	: FMassRawEntityInChunkData(Source.ChunkRawMemory, InIndexWithinChunk)
+		, ChunkIndex(Source.ChunkIndex), ChunkSerialNumber(Source.ChunkSerialNumber)
+{
+}
+
 inline bool FMassEntityInChunkDataHandle::operator==(const FMassEntityInChunkDataHandle& Other) const
 {
 	return FMassRawEntityInChunkData::operator==(Other)
diff --git a/Engine/Source/Runtime/MassEntity/Public/MassElement.h b/Engine/Source/Runtime/MassEntity/Public/MassElement.h
--- a/Engine/Source/Runtime/MassEntity/Public/MassElement.h
+++ b/Engine/Source/Runtime/MassEntity/Public/MassElement.h
@@ -88,12 +88,7 @@ namespace UE::Mass
 			return EElementType::MAX;
 		}
 
-		static void OnTypeRegistered(const TNotNull<const UStruct*> Type, int32 Index)
-		{
-			const EElementType ElementType = DetermineElementType(Type);
-			GetTypeBitArray(ElementType).AddAtIndex(Index);
-			bTypeBitSetsDirty = true;
-		}
+		static void OnTypeRegistered(const TNotNull<const UStruct*> Type, int32 Index);
 
 		static void OnModuleInitialized()
 		{
@@ -105,33 +100,26 @@ namespace UE::Mass
 			FElementBitSetBaseStructTrackerWrapper::StructTracker.OnTypeRegistered.Remove(OnTypeRegisteredDelegateHandle);
 		}
 
-		UE_FORCEINLINE_HINT static void UpdateTypeBitSets()
-		{
-			if (bTypeBitSetsDirty)
-			{
-				AllSharedFragmentsBitSet = GetTypeBitArray(EElementType::ConstSharedFragment) | GetTypeBitArray(EElementType::SharedFragment);
-				AllFragmentsAndTagsBitSet = (GetTypeBitArray(EElementType::Fragment) | GetTypeBitArray(EElementType::Tag));
-				bTypeBitSetsDirty = false;
-			}
-		}
-
 		static const FElementBitSet& GetAllFragmentsAndTags()
 		{
-			UpdateTypeBitSets();
 			return AllFragmentsAndTagsBitSet;
 		}
 
 		static const FElementBitSet& GetAllSharedFragments()
 		{
-			UpdateTypeBitSets();
 			return AllSharedFragmentsBitSet;
 		}
 
+		static const FElementBitSet& GetAllSparseElements()
+		{
+			return AllSparseElementsBitSet;
+		}
+
 	protected:
 		UE_API static TStaticArray<FElementBitSetBase, static_cast<uint8>(EElementType::MAX)> ElementTypeBitArrays;
 		UE_API static FElementBitSet AllSharedFragmentsBitSet;
 		UE_API static FElementBitSet AllFragmentsAndTagsBitSet;
-		UE_API static bool bTypeBitSetsDirty;
+		UE_API static FElementBitSet AllSparseElementsBitSet;
 		static FDelegateHandle OnTypeRegisteredDelegateHandle;
 
 	public:
diff --git a/Engine/Source/Runtime/MassEntity/Public/MassEntityBuilder.h b/Engine/Source/Runtime/MassEntity/Public/MassEntityBuilder.h
--- a/Engine/Source/Runtime/MassEntity/Public/MassEntityBuilder.h
+++ b/Engine/Source/Runtime/MassEntity/Public/MassEntityBuilder.h
@@ -21,6 +21,12 @@ namespace UE::Mass
 		struct FEntityBuilderHelper;
 	}
 
+	enum class EIncludeSparseElements : uint8
+	{
+		No,
+		Yes
+	};
+
 /**
  * FEntityBuilder is a utility struct that provides a convenient way to create and configure entities
  * in the Mass framework. It bridges multiple APIs from FMassEntityManager, MassSpawnerSubsystem,
@@ -132,14 +138,14 @@ struct FEntityBuilder
 	 * @param SourceEntityHandle valid handle for a fully constructed, built entity.
 	 * @return whether the operation was successful
 	 */
-	UE_API bool AppendDataFromEntity(const FMassEntityHandle SourceEntityHandle);
+	UE_API bool AppendDataFromEntity(const FMassEntityHandle SourceEntityHandle, EIncludeSparseElements IncludeSparseElements = EIncludeSparseElements::Yes);
 
 	/**
 	 * Copies all element types and values stored by the entity indicated by SourceEntityHandle. Any existing builder data will be overridden
 	 * @param SourceEntityHandle valid handle for a fully constructed, built entity.
 	 * @return whether the operation was successful
 	 */
-	UE_API bool CopyDataFromEntity(const FMassEntityHandle SourceEntityHandle);
+	UE_API bool CopyDataFromEntity(const FMassEntityHandle SourceEntityHandle, EIncludeSparseElements IncludeSparseElements = EIncludeSparseElements::Yes);
 
 	/**
 	 * Adds a fragment of type T to the entity and returns a reference to it, constructing it with the provided arguments.
@@ -315,7 +321,8 @@ private:
 	void CacheEntityHandle() const;
 
 	static void CreateEntityImpl(FMassEntityManager& EntityManager, FMassEntityHandle EntityHandle, const FMassArchetypeHandle& ArchetypeHandle
-		, const FMassArchetypeSharedFragmentValues& SharedFragmentValues, TConstArrayView<FInstancedStruct> Fragments, TConstArrayView<FPendingRelationParams> RelationsParams);
+		, const FMassArchetypeSharedFragmentValues& SharedFragmentValues, TConstArrayView<FInstancedStruct> Fragments, TConstArrayView<FPendingRelationParams> RelationsParams
+		, const FMassElementBitSet& CompositionBitSet);
 
 	TArray<FInstancedStruct> Fragments;
 	TArray<FInstancedStruct> SharedFragments;
diff --git a/Engine/Source/Runtime/MassEntity/Public/MassEntityConcepts.h b/Engine/Source/Runtime/MassEntity/Public/MassEntityConcepts.h
--- a/Engine/Source/Runtime/MassEntity/Public/MassEntityConcepts.h
+++ b/Engine/Source/Runtime/MassEntity/Public/MassEntityConcepts.h
@@ -40,6 +40,9 @@ namespace UE::Mass
 	template<typename T>
 	concept CSubsystem = TIsDerivedFrom<Clean<T>, USubsystem>::Value;
 
+	template<typename T>
+	concept CSparse = TIsDerivedFrom<Clean<T>, FMassSparseFragment>::Value || TIsDerivedFrom<Clean<T>, FMassSparseTag>::Value;
+
 	namespace Private
 	{
 		template<CElement T>
@@ -56,4 +59,11 @@ namespace UE::Mass
 
 	template<typename T>
 	using TElementType = typename Private::TElementTypeHelper<T>::Type;
+
+	template<typename T>
+	bool IsSparse()
+	{
+		constexpr bool bIsSparse = CSparse<typename TRemoveReference<T>::Type>;
+		return bIsSparse;
+	}
 }
diff --git a/Engine/Source/Runtime/MassEntity/Public/MassEntityElementTypes.h b/Engine/Source/Runtime/MassEntity/Public/MassEntityElementTypes.h
--- a/Engine/Source/Runtime/MassEntity/Public/MassEntityElementTypes.h
+++ b/Engine/Source/Runtime/MassEntity/Public/MassEntityElementTypes.h
@@ -10,7 +10,7 @@ struct FMassElement
 	GENERATED_BODY()
 };
 
-// This is the base class for all lightweight fragments
+/** This is the base class for all Mass fragments */
 USTRUCT()
 struct FMassFragment : public FMassElement
 {
@@ -22,15 +22,36 @@ struct FMassFragment : public FMassElement
 #define MASS_INVALID_FRAGMENT_MSG  "Given struct doesn't represent a valid fragment type." _MASS_INVALID_FRAGMENT_CORE_MESSAGE
 #define MASS_INVALID_FRAGMENT_MSG_F  "Type %s is not a valid fragment type." _MASS_INVALID_FRAGMENT_CORE_MESSAGE
 
+/**
+ * Special-case base class for fragments. Derive from this if you intend to use the fragment type as a sparse fragment
+ * See UE::Mass::FSparseElementsStorage for more details on sparse elements.
+ */
+USTRUCT()
+struct FMassSparseFragment : public FMassFragment
+{
+	GENERATED_BODY()
+};
 
-// This is the base class for types that will only be tested for presence/absence, i.e. Tags.
-// Subclasses should never contain any member properties.
+/**
+ * This is the base class for types that will only be tested for presence/absence, i.e. Tags.
+ * Subclasses should never contain any member properties.
+ */
 USTRUCT()
 struct FMassTag : public FMassElement
 {
 	GENERATED_BODY()
 };
 
+/**
+ * Special-case base class for tags. Derive from this if you intend to use the tag type as a sparse tag
+ * See UE::Mass::FSparseElementsStorage for more details on sparse elements.
+ */
+USTRUCT()
+struct FMassSparseTag : public FMassTag
+{
+	GENERATED_BODY()
+};
+
 USTRUCT()
 struct FMassChunkFragment : public FMassElement
 {
@@ -70,12 +91,24 @@ namespace UE::Mass
 		return Struct && Struct->IsChildOf(FMassFragment::StaticStruct());
 	}
 
+	template<>
+	inline bool IsA<FMassSparseFragment>(const UStruct* Struct)
+	{
+		return Struct && Struct->IsChildOf(FMassSparseFragment::StaticStruct());
+	}
+
 	template<>
 	inline bool IsA<FMassTag>(const UStruct* Struct)
 	{
 		return Struct && Struct->IsChildOf(FMassTag::StaticStruct());
 	}
 
+	template<>
+	inline bool IsA<FMassSparseTag>(const UStruct* Struct)
+	{
+		return Struct && Struct->IsChildOf(FMassSparseTag::StaticStruct());
+	}
+
 	template<>
 	inline bool IsA<FMassChunkFragment>(const UStruct* Struct)
 	{
@@ -93,4 +126,9 @@ namespace UE::Mass
 	{
 		return Struct && Struct->IsChildOf(FMassConstSharedFragment::StaticStruct());
 	}
+
+	inline bool IsSparse(TNotNull<const UStruct*> Struct)
+	{
+		return Struct->IsChildOf(FMassSparseFragment::StaticStruct()) || Struct->IsChildOf(FMassSparseTag::StaticStruct());
+	}
 }
diff --git a/Engine/Source/Runtime/MassEntity/Public/MassEntityManager.h b/Engine/Source/Runtime/MassEntity/Public/MassEntityManager.h
--- a/Engine/Source/Runtime/MassEntity/Public/MassEntityManager.h
+++ b/Engine/Source/Runtime/MassEntity/Public/MassEntityManager.h
@@ -19,6 +19,7 @@
 #include "MassArchetypeGroup.h"
 #include "MassRelationManager.h"
 #include "MassEntityRelations.h"
+#include "MassSparseElementsStorage.h"
 
 #if WITH_MASSENTITY_DEBUG
 #include "Misc/EnumClassFlags.h"
@@ -403,6 +404,45 @@ public:
 	UE_API void RemoveTagFromEntity(FMassEntityHandle EntityHandle, TNotNull<const UScriptStruct*> TagType);
 	UE_API void SwapTagsForEntity(FMassEntityHandle EntityHandle, TNotNull<const UScriptStruct*> FromFragmentType, TNotNull<const UScriptStruct*> ToFragmentType);
 
+	/**
+	 * Adds sparse element indicated by ElementType to the entity indicated by EntityHandle.
+	 * Works with tags too, but returns an empty struct view
+	 * For more information on Sparse Elements see FSparseElementsStorage
+	 */
+	UE_API FStructView AddSparseElementToEntity(FMassEntityHandle EntityHandle, TNotNull<const UScriptStruct*> ElementType);
+
+	/**
+	 * Adds a sparse fragment to the entity and returns a reference to the newly created fragments.
+	 * This function can only fail if EntityHandle doesn't represent a valid, active entity.
+	 */
+	template<UE::Mass::CFragment T>
+	requires UE::Mass::CSparse<T>
+	T& AddSparseElementToEntity(FMassEntityHandle EntityHandle)
+	{
+		CheckIfEntityIsActive(EntityHandle);
+		return AddSparseElementToEntity(EntityHandle, T::StaticStruct()).template Get<T>();
+	}
+
+	template<UE::Mass::CTag T>
+	requires UE::Mass::CSparse<T>
+	void AddSparseElementToEntity(FMassEntityHandle EntityHandle)
+	{
+		FStructView _ = AddSparseElementToEntity(EntityHandle, T::StaticStruct());
+	}
+
+	UE_API bool RemoveSparseElementFromEntity(FMassEntityHandle EntityHandle, TNotNull<const UScriptStruct*> ElementType);
+	template<UE::Mass::CSparse T>
+	bool RemoveSparseElementFromEntity(FMassEntityHandle EntityHandle)
+	{
+		return RemoveSparseElementFromEntity(EntityHandle, T::StaticStruct());
+	}
+
+	/**
+	 * Adds sparse elements indicated by ElementsBitset to the entity indicated by EntityHandle.
+	 * For more information on Sparse Elements see FSparseElementsStorage
+	 */
+	UE_API void AddSparseElementsToEntity(FMassEntityHandle EntityHandle, const FMassElementBitSet& ElementsBitset);
+
 	/**
 	 * Adds ElementType to the entities, treating it accordingly based on which element type it represents (i.e. Fragment or Tag).
 	 * The function will assert on unhandled types.
@@ -411,6 +451,7 @@ public:
 
 	/**
 	 * Adds ElementType to the entity, treating it accordingly based on which element type it represents (i.e. Fragment or Tag).
+	 * This function handles sparse elements.
 	 * The function will assert on unhandled types.
 	 */
 	void AddElementToEntity(FMassEntityHandle Entity, TNotNull<const UScriptStruct*> ElementType);
@@ -430,6 +471,12 @@ public:
 	/** @return whether Entity has an element of ElementType */
 	UE_API bool DoesEntityHaveElement(FMassEntityHandle Entity, TNotNull<const UScriptStruct*> ElementType) const;
 
+	template<UE::Mass::CElement T>
+	bool DoesEntityHaveElement(FMassEntityHandle Entity) const
+	{
+		return DoesEntityHaveElement(Entity, T::StaticStruct());
+	}
+
 	/** 
 	 * Adds a new const shared fragment to the given entity. Note that it only works if the given entity doesn't have
 	 * a shared fragment of the given type. The function will give a soft "pass" if the entity has the shared fragment
@@ -512,6 +559,10 @@ public:
 	 */
 	UE_API void BatchAddSharedFragmentsForEntities(TConstArrayView<FMassArchetypeEntityCollection> EntityCollections, const FMassArchetypeSharedFragmentValues& AddedFragmentValues);
 
+	UE_API void BatchAddSparseElementToEntities(TConstArrayView<FMassArchetypeEntityCollection> EntityCollections, TNotNull<const UScriptStruct*> ElementType);
+	// @todo add a version with FMassArchetypeEntityCollectionWithPayload to add pre-configured fragment instances
+	UE_API void BatchRemoveSparseElementFromEntities(TConstArrayView<FMassArchetypeEntityCollection> EntityCollections, TNotNull<const UScriptStruct*> ElementType);
+
 	/**
 	 * Adds elements indicated by InOutDescriptor to the entity indicated by EntityHandle. The function also figures out which elements
 	 * in InOutDescriptor are missing from the current composition of the given entity and then returns the resulting 
@@ -700,6 +751,11 @@ public:
 		}
 	}
 
+	UE_API FStructView GetMutableSparseElementDataForEntity(TNotNull<const UScriptStruct*> ElementType, const FMassEntityHandle EntityHandle);
+	UE_API FConstStructView GetSparseElementDataForEntity(TNotNull<const UScriptStruct*> ElementType, const FMassEntityHandle EntityHandle) const;
+
+	FMassElementBitSet GetSparseElementsBitSetForEntity(const FMassEntityHandle EntityHandle) const;
+
 	uint32 GetArchetypeDataVersion() const { return ArchetypeDataVersion; }
 
 	/**
@@ -772,6 +828,24 @@ private:
 		return ConstSharedFragmentsContainer.FindOrAdd<T>(Hash, Forward<TArgs>(InArgs)...);
 	}
 
+	/** 
+	 * Creates UE::Mass::FSparseElementIterator instance that can be used to iterate over all instances of a given sparse element fragment
+	 * Note: works only for sparse fragments, as tags contain no data. 
+	 */
+	UE::Mass::FSparseElementIterator CreateSparseElementsIterator(TNotNull<const UScriptStruct*> ElementType)
+	{
+		return UE::Mass::IsA<FMassSparseFragment>(ElementType)
+			? SparseElementsStorage.CreateElementIterator(ElementType)
+			: UE::Mass::FSparseElementIterator();
+	}
+
+	template<UE::Mass::CSparse T>
+	requires UE::Mass::CFragment<T>
+	UE::Mass::FSparseElementIterator CreateSparseElementsIterator()
+	{
+		return SparseElementsStorage.CreateElementIterator(T::StaticStruct());
+	}
+
 	/**
 	 * Creates fully built entity ready to be used by the subsystem
 	 * Does not check for synchronous execution and should only be used by processors through the execution context
@@ -1172,6 +1246,8 @@ public:
 	UE_API const UE::Mass::FStorageType& DebugGetEntityStorageInterface() const;
 
 	UE_API bool DebugHasCommandsToFlush() const;
+
+	UE_API UE::Mass::FSparseElementsStorage& DebugGetSparseElementsStorage();
 #endif // WITH_MASSENTITY_DEBUG
 
 protected:
@@ -1365,6 +1441,9 @@ private:
 	TSharedRef<UE::Mass::FTypeManager> TypeManager;
 	UE::Mass::FRelationManager RelationManager;
 
+	/** Stores sparse elements instances. Not accessible from outside on purpose. */
+	UE::Mass::FSparseElementsStorage SparseElementsStorage;
+
 	TMap<const FName, const int32> GroupNameToTypeIndex;
 	// @todo we'll probably have some "GroupTypeInformation" here in the future
 	TArray<const FName> GroupTypes;
diff --git a/Engine/Source/Runtime/MassEntity/Public/MassEntityView.h b/Engine/Source/Runtime/MassEntity/Public/MassEntityView.h
--- a/Engine/Source/Runtime/MassEntity/Public/MassEntityView.h
+++ b/Engine/Source/Runtime/MassEntity/Public/MassEntityView.h
@@ -76,6 +76,20 @@ struct FMassEntityView
 		return FStructView(FragmentType, static_cast<uint8*>(GetFragmentPtr(*FragmentType)));
 	}
 
+	template<UE::Mass::CSparse T>
+	requires UE::Mass::CFragment<T>
+	const T* GetSparseFragmentDataPtr() const
+	{
+		return GetSparseFragmentPtr(T::StaticStruct()).template GetPtr<T>();
+	}
+
+	template<UE::Mass::CSparse T>
+	requires UE::Mass::CFragment<T>
+	T* GetSparseFragmentDataPtr()
+	{
+		return GetSparseFragmentPtr(T::StaticStruct()).template GetPtr<T>();
+	}
+
 	/** if the viewed entity doesn't have the given const shared fragment the function will return null */
 	template<typename T>
 	const T* GetConstSharedFragmentDataPtr() const
@@ -123,13 +137,30 @@ struct FMassEntityView
 	}
 
 	template<typename T>
-	bool HasTag() const
+	bool HasTag(bool bIncludeSparse = true) const
+	{
+		return HasTag(*T::StaticStruct(), bIncludeSparse);
+	}
+
+	UE_API bool HasTag(TNotNull<const UScriptStruct*> TagType, bool bIncludeSparse = true) const;
+	bool HasTag(const UScriptStruct& TagType, bool bIncludeSparse = true) const
 	{
-		static_assert(UE::Mass::CTag<T>, "Given struct doesn't represent a valid tag type. Make sure to inherit from FMassTag or one of its child-types.");
-		return HasTag(*T::StaticStruct());
+		return HasTag(&TagType, bIncludeSparse);
 	}
 
-	UE_API bool HasTag(const UScriptStruct& TagType) const;
+	UE_API bool HasFragment(TNotNull<const UScriptStruct*> FragmentType, bool bIncludeSparse = true) const;
+	template<UE::Mass::CFragment T>
+	bool HasFragment(bool bIncludeSparse = true) const
+	{
+		return HasFragment(T::StaticStruct(), bIncludeSparse);
+	}
+
+	UE_API bool HasElement(TNotNull<const UScriptStruct*> ElementType, bool bIncludeSparse = true) const;
+	template<UE::Mass::CSparse T>
+	bool HasElement(bool bIncludeSparse = true) const
+	{
+		return HasElement(T::StaticStruct(), bIncludeSparse);
+	}
 
 	bool IsSet() const;
 	bool IsValid() const;
@@ -142,6 +173,8 @@ protected:
 	UE_API const void* GetConstSharedFragmentPtrChecked(const UScriptStruct& FragmentType) const;
 	UE_API void* GetSharedFragmentPtr(const UScriptStruct& FragmentType) const;
 	UE_API void* GetSharedFragmentPtrChecked(const UScriptStruct& FragmentType) const;
+	UE_API FStructView GetSparseFragmentPtr(TNotNull<const UScriptStruct*> FragmentType);
+	UE_API FConstStructView GetSparseFragmentPtr(TNotNull<const UScriptStruct*> FragmentType) const;
 
 private:
 	FMassEntityHandle Entity;
diff --git a/Engine/Source/Runtime/MassEntity/Public/MassExecutionContext.h b/Engine/Source/Runtime/MassEntity/Public/MassExecutionContext.h
--- a/Engine/Source/Runtime/MassEntity/Public/MassExecutionContext.h
+++ b/Engine/Source/Runtime/MassEntity/Public/MassExecutionContext.h
@@ -16,21 +16,29 @@
 #include "MassEntityView.h"
 
 
-#define CHECK_IF_VALID(View, Type) \
+#define MASS_CHECK_IF_VALID(View, Type) \
 	checkf(View \
 		, TEXT("Requested fragment type not bound, type %s. Make sure it has been listed as required."), *GetNameSafe(Type))
 
-#define CHECK_IF_READWRITE(View) \
+#define MASS_CHECK_IF_READWRITE(View) \
 	checkf(View == nullptr || View->Requirement.AccessMode == EMassFragmentAccess::ReadWrite \
 		, TEXT("Requested fragment type not bound for writing, type %s. Make sure it has been listed as required in ReadWrite mode.") \
 		, View ? *GetNameSafe(View->Requirement.StructType) : TEXT("[Not found]"))
 
+#define MASS_CHECK_IF_SPARSE_REQUIREMENT(ElementType) \
+	checkf(RequestedSparseElementsAccess.Contains(*ElementType) \
+			, TEXT("%hs: %s not requested as requirement by the query"), __FUNCTION__, *ElementType->GetName())
+
 struct FMassEntityQuery;
 struct FMassEntityLinkFragment;
+struct FMassArchetypeChunk;
+struct FMassArchetypeData;
 
 struct FMassExecutionContext
 {
+	struct FSparseEntityIterator;
 private:
+	friend FMassArchetypeData;
 
 	template< typename ViewType >
 	struct TFragmentView 
@@ -55,6 +63,10 @@ private:
 	using FSharedFragmentView = TFragmentView<FStructView>;
 	TArray<FSharedFragmentView, TInlineAllocator<4>> SharedFragmentViews;
 
+	FMassElementBitSet RequestedSparseElementsAccess;
+	// @todo could add a construct similar to TFragmentView that would cache direct access to FSparseElementStorage::FTypePool
+	// for requested sparse element types. That would be a perf optimization
+
 	/** linked entity to allow access to fragments on an entity not in the query, must be bound before access */
 	FMassEntityView LinkedEntityView;
 
@@ -74,7 +86,6 @@ private:
 	/** @todo rename to "payload" */
 	FInstancedStruct AuxData;
 	float DeltaTimeSeconds = 0.0f;
-	int32 ChunkSerialModificationNumber = -1;
 	FMassElementBitSet CurrentArchetypeCompositionBitSet;
 #if WITH_MASSENTITY_DEBUG
 	FColor DebugColor;
@@ -113,6 +124,19 @@ private:
 	/** Track the serial number for FEntityIterator creation */
 	uint32 IteratorSerialNumberGenerator = 0;
 
+	EMassExecutionContextType ExecutionType = EMassExecutionContextType::Local;
+
+	/** Used to control when the context is allowed to flush commands collected in DeferredCommandBuffer. This mechanism 
+	 * is mainly utilized to avoid numerous small flushes in favor of fewer larger ones. */
+	bool bFlushDeferredCommands = true;
+
+	/** cached value of current query's HasSparseRequirements() */
+	bool bHasSparseRequirements = false;
+
+	/** Extremely transient, needs to be kept private. Used for direct access to sparse elements API  */
+	FMassArchetypeData* ArchetypeDataPtr = nullptr;
+	FMassEntityInChunkDataHandle ChunkEntityIndex;
+
 #if WITH_MASSENTITY_DEBUG
 	FString DebugExecutionDescription;
 
@@ -120,10 +144,6 @@ private:
 	// js@todo make this more generic
 	TWeakObjectPtr<UMassProcessor> DebugProcessor;
 #endif // WITH_MASSENTITY_DEBUG
-	
-	/** Used to control when the context is allowed to flush commands collected in DeferredCommandBuffer. This mechanism 
-	 * is mainly utilized to avoid numerous small flushes in favor of fewer larger ones. */
-	bool bFlushDeferredCommands = true;
 
 	/** Active permissions for archetype creation (processor-level when no query is active; query-level otherwise) */
 	const FMassEntityCreationRequirements* ActiveEntityCreationRequirements = nullptr;
@@ -143,9 +163,7 @@ private:
 		SubsystemAccess.SetSubsystemRequirementBits(InConstSubsystemsBitSet, InMutableSubsystemsBitSet);
 	}
 
-	EMassExecutionContextType ExecutionType = EMassExecutionContextType::Local;
-
-	friend FMassArchetypeData;
+	MASSENTITY_API FStructView GetSparseElementInternal(TNotNull<const UScriptStruct*> ElementType, const FMassEntityHandle EntityHandle) const;
 
 	/**
 	 * Note that this operator is private on purpose, used to simplify the implementation of constructors.
@@ -182,6 +200,12 @@ public:
 	void ApplyFragmentRequirements(const FMassEntityQuery& RequestingQuery);
 	void ClearFragmentViews(const FMassEntityQuery& RequestingQuery);
 
+	/**
+	 * @return whether the entity indicated via FSparseEntityIterator meets the sparse element requirements.
+	 *		Will always return true if there are no sparse requirements
+	 */
+	bool ExecuteSparseElementsFilter(const FSparseEntityIterator& EntityIt) const;
+
 	/**
 	 * Iterator to easily loop through entities in the current chunk.
 	 * Supports ranged for and can be used directly as an entity index for the current chunk.
@@ -254,7 +278,12 @@ public:
 		 */
 		FEntityIterator(const FEntityIterator&) = delete;
 
-	private:
+		FMassEntityHandle GetEntityHandle() const
+		{
+			return bool(*this) ? ExecutionContext.GetEntities()[EntityIndex] : FMassEntityHandle();
+		}
+
+	protected:
 		friend FMassExecutionContext;
 		FEntityIterator(FMassExecutionContext& InExecutionContext);
 		FEntityIterator(FMassExecutionContext& InExecutionContext, FQueryTransientRuntime& InQueryRuntime);
@@ -276,6 +305,67 @@ public:
 	 */
 	MASSENTITY_API FEntityIterator CreateEntityIterator();
 
+	/**
+	 * A flavor of FEntityIterator that skips all the entities that don't match current query's sparse element requirements.
+	 * This means that if there are no sparse requirements then FSparseEntityIterator will behave like a regular FEntityIterator
+	 * but a less efficiently so.
+	 */
+	struct FSparseEntityIterator : FEntityIterator
+	{
+		/** Note that the default constructor results in an "invalid" (i.e. useless) itertor. */
+		FSparseEntityIterator() = default;
+
+		explicit FSparseEntityIterator(FEntityIterator&& Source)
+			: FEntityIterator(MoveTemp(Source))
+		{
+			if (EntityIndex < NumEntities && ExecutionContext.ExecuteSparseElementsFilter(*this) == false)
+			{
+				++(*this);
+			}
+		}
+
+		FSparseEntityIterator& operator++()
+		{
+			do
+			{
+				++EntityIndex;
+			} while (EntityIndex < NumEntities && ExecutionContext.ExecuteSparseElementsFilter(*this) == false);
+
+#if WITH_MASSENTITY_DEBUG
+			if (UNLIKELY(QueryRuntime.bCheckProcessorBreaks || QueryRuntime.BreakFragmentsCount != 0) 
+				&& EntityIndex < NumEntities)
+			{
+				TestBreakpoints();
+			}
+#endif //WITH_MASSENTITY_DEBUG
+			return *this;
+		}
+
+		void operator++(int)
+		{
+			++(*this);
+		}
+
+		FSparseEntityIterator&& begin()
+		{
+			return MoveTemp(*this);
+		}
+
+		FSparseEntityIterator end() const
+		{
+			FSparseEntityIterator End;
+			End.EntityIndex = NumEntities;
+			return End;
+		}
+	};
+
+	/**
+	 * Creates an entity iterator for the current chunk that's configured to skip entities that are not matching
+	 * the current query's sparse element requirements.
+	 */
+	MASSENTITY_API FSparseEntityIterator CreateSparseEntityIterator();
+	// @todo we need AnySparseElementsBitSet and NoneSparseElementsBitSet
+
 	/** Sets bFlushDeferredCommands. Note that setting to True while the system is being executed doesn't result in
 	 *  immediate commands flushing */
 	void SetFlushDeferredCommands(const bool bNewFlushDeferredCommands);
@@ -356,9 +446,14 @@ public:
 	}
 #endif // WITH_MASSENTITY_DEBUG
 
-	/** Chunk related operations */
-	void SetCurrentChunkSerialModificationNumber(const int32 SerialModificationNumber) { ChunkSerialModificationNumber = SerialModificationNumber; }
-	int32 GetChunkSerialModificationNumber() const { return ChunkSerialModificationNumber; }
+	//-----------------------------------------------------------------------------
+	// Chunk related operations
+	//-----------------------------------------------------------------------------
+
+	/** @return whether the given chunk is fine to be processed, i.e. `false` indicates the chunk should be skipped */
+	bool SetCurrentChunk(int32 ChunkIndex, const FMassArchetypeChunk& Chunk);
+
+	int32 GetChunkSerialModificationNumber() const;
 
 	template<typename T>
 	T* GetMutableChunkFragmentPtr()
@@ -367,7 +462,7 @@ public:
 
 		const UScriptStruct* Type = T::StaticStruct();
 		FChunkFragmentView* FoundChunkFragmentData = ChunkFragmentViews.FindByPredicate([Type](const FChunkFragmentView& Element) { return Element.Requirement.StructType == Type; } );
-		CHECK_IF_READWRITE(FoundChunkFragmentData);
+		MASS_CHECK_IF_READWRITE(FoundChunkFragmentData);
 		return FoundChunkFragmentData ? FoundChunkFragmentData->FragmentView.GetPtr<T>() : static_cast<T*>(nullptr);
 	}
 	
@@ -375,7 +470,7 @@ public:
 	T& GetMutableChunkFragment()
 	{
 		T* ChunkFragment = GetMutableChunkFragmentPtr<T>();
-		CHECK_IF_VALID(ChunkFragment, T::StaticStruct());
+		MASS_CHECK_IF_VALID(ChunkFragment, T::StaticStruct());
 		return *ChunkFragment;
 	}
 
@@ -393,7 +488,7 @@ public:
 	const T& GetChunkFragment() const
 	{
 		const T* ChunkFragment = GetChunkFragmentPtr<T>();
-		CHECK_IF_VALID(ChunkFragment, T::StaticStruct());
+		MASS_CHECK_IF_VALID(ChunkFragment, T::StaticStruct());
 		return *ChunkFragment;
 	}
 
@@ -417,7 +512,7 @@ public:
 	const T& GetConstSharedFragment() const
 	{
 		const T* SharedFragment = GetConstSharedFragmentPtr<const T>();
-		CHECK_IF_VALID(SharedFragment, T::StaticStruct());
+		MASS_CHECK_IF_VALID(SharedFragment, T::StaticStruct());
 		return *SharedFragment;
 	}
 
@@ -427,7 +522,7 @@ public:
 		static_assert(UE::Mass::CSharedFragment<T>, "Given struct doesn't represent a valid shared fragment type. Make sure to inherit from FMassSharedFragment or one of its child-types.");
 
 		FSharedFragmentView* FoundSharedFragmentData = SharedFragmentViews.FindByPredicate([](const FSharedFragmentView& Element) { return Element.Requirement.StructType == T::StaticStruct(); });
-		CHECK_IF_READWRITE(FoundSharedFragmentData);
+		MASS_CHECK_IF_READWRITE(FoundSharedFragmentData);
 		return FoundSharedFragmentData ? FoundSharedFragmentData->FragmentView.GetPtr<T>() : static_cast<T*>(nullptr);
 	}
 
@@ -444,7 +539,7 @@ public:
 	T& GetMutableSharedFragment()
 	{
 		T* SharedFragment = GetMutableSharedFragmentPtr<T>();
-		CHECK_IF_VALID(SharedFragment, T::StaticStruct());
+		MASS_CHECK_IF_VALID(SharedFragment, T::StaticStruct());
 		return *SharedFragment;
 	}
 
@@ -452,7 +547,7 @@ public:
 	const T& GetSharedFragment() const
 	{
 		const T* SharedFragment = GetSharedFragmentPtr<T>();
-		CHECK_IF_VALID(SharedFragment, T::StaticStruct());
+		MASS_CHECK_IF_VALID(SharedFragment, T::StaticStruct());
 		return *SharedFragment;
 	}
 
@@ -462,8 +557,8 @@ public:
 	{
 		const UScriptStruct* FragmentType = TFragment::StaticStruct();
 		const FFragmentView* View = FragmentViews.FindByPredicate([FragmentType](const FFragmentView& Element) { return Element.Requirement.StructType == FragmentType; });
-		CHECK_IF_VALID(View, FragmentType);
-		CHECK_IF_READWRITE(View);
+		MASS_CHECK_IF_VALID(View, FragmentType);
+		MASS_CHECK_IF_READWRITE(View);
 		return MakeArrayView<TFragment>((TFragment*)View->FragmentView.GetData(), View->FragmentView.Num());
 	}
 
@@ -472,22 +567,22 @@ public:
 	{
 		const UScriptStruct* FragmentType = TFragment::StaticStruct();
 		const FFragmentView* View = FragmentViews.FindByPredicate([FragmentType](const FFragmentView& Element) { return Element.Requirement.StructType == FragmentType; });
-		CHECK_IF_VALID(View, TFragment::StaticStruct());
+		MASS_CHECK_IF_VALID(View, TFragment::StaticStruct());
 		return TConstArrayView<TFragment>((const TFragment*)View->FragmentView.GetData(), View->FragmentView.Num());
 	}
 
 	TConstArrayView<FMassFragment> GetFragmentView(const UScriptStruct* FragmentType) const
 	{
 		const FFragmentView* View = FragmentViews.FindByPredicate([FragmentType](const FFragmentView& Element) { return Element.Requirement.StructType == FragmentType; });
-		CHECK_IF_VALID(View, FragmentType);
+		MASS_CHECK_IF_VALID(View, FragmentType);
 		return TConstArrayView<FMassFragment>((const FMassFragment*)View->FragmentView.GetData(), View->FragmentView.Num());;
 	}
 
 	TArrayView<FMassFragment> GetMutableFragmentView(const UScriptStruct* FragmentType) 
 	{
 		const FFragmentView* View = FragmentViews.FindByPredicate([FragmentType](const FFragmentView& Element) { return Element.Requirement.StructType == FragmentType; });
-		CHECK_IF_VALID(View, FragmentType);
-		CHECK_IF_READWRITE(View);
+		MASS_CHECK_IF_VALID(View, FragmentType);
+		MASS_CHECK_IF_READWRITE(View);
 		return View->FragmentView;
 	}
 
@@ -607,6 +702,75 @@ public:
 		return SubsystemAccess.GetSubsystemChecked<T>(SubsystemClass);
 	}
 
+	//-----------------------------------------------------------------------------
+	// Sparse elements querying
+	//-----------------------------------------------------------------------------
+	template<UE::Mass::CSparse TFragment>
+	requires UE::Mass::CFragment<TFragment>
+	TFragment* GetMutableSparseElement(const FEntityIterator& EntityIterator) const
+	{
+		MASS_CHECK_IF_SPARSE_REQUIREMENT(TFragment::StaticStruct());
+
+		FStructView SparseElementView = GetMutableSparseElement(TFragment::StaticStruct(), EntityIterator);
+		return SparseElementView.GetPtr<TFragment>();
+	}
+
+	template<UE::Mass::CSparse TFragmentBase>
+	requires UE::Mass::CFragment<TFragmentBase>
+	TFragmentBase* GetMutableSparseElement(TNotNull<const UScriptStruct*> ElementType, const FEntityIterator& EntityIterator) const
+	{
+		checkf(ElementType->IsChildOf(TFragmentBase::StaticStruct()), TEXT("%s is not derived from %s")
+			, *ElementType->GetName(), *TFragmentBase::StaticStruct()->GetName());
+		MASS_CHECK_IF_SPARSE_REQUIREMENT(ElementType);
+
+		FStructView SparseElementView = GetSparseElementInternal(ElementType, EntityIterator.GetEntityHandle());
+		return SparseElementView.GetPtr<TFragmentBase>();
+	}	
+
+	FStructView GetMutableSparseElement(TNotNull<const UScriptStruct*> ElementType, const FEntityIterator& EntityIterator) const
+	{
+		checkf(UE::Mass::IsSparse(ElementType), TEXT("%s doesn't represent a sparse element type"), *ElementType->GetName());
+		MASS_CHECK_IF_SPARSE_REQUIREMENT(ElementType);
+
+		return UE::Mass::IsA<FMassFragment>(ElementType)
+			? GetSparseElementInternal(ElementType, EntityIterator.GetEntityHandle())
+			: FStructView();
+	}
+
+	template<UE::Mass::CSparse TFragment>
+	requires UE::Mass::CFragment<TFragment>
+	const TFragment* GetSparseElement(const FEntityIterator& EntityIterator) const
+	{
+		MASS_CHECK_IF_SPARSE_REQUIREMENT(TFragment::StaticStruct());
+		FConstStructView SparseElementView = GetSparseElement(TFragment::StaticStruct(), EntityIterator);
+		return SparseElementView.GetPtr<const TFragment>();
+	}
+
+	template<UE::Mass::CSparse TFragmentBase>
+	requires UE::Mass::CFragment<TFragmentBase>
+	const TFragmentBase* GetSparseElement(TNotNull<const UScriptStruct*> ElementType, const FEntityIterator& EntityIterator) const
+	{
+		checkf(ElementType->IsChildOf(TFragmentBase::StaticStruct()), TEXT("%s is not derived from %s")
+			, *ElementType->GetName(), *TFragmentBase::StaticStruct()->GetName());
+		MASS_CHECK_IF_SPARSE_REQUIREMENT(ElementType);
+
+		FConstStructView SparseElementView = GetSparseElementInternal(ElementType, EntityIterator.GetEntityHandle());
+		return SparseElementView.GetPtr<const TFragmentBase>();
+	}
+
+	FConstStructView GetSparseElement(TNotNull<const UScriptStruct*> ElementType, const FEntityIterator& EntityIterator) const
+	{
+		checkf(UE::Mass::IsSparse(ElementType), TEXT("%s doesn't represent a sparse element type"), *ElementType->GetName());
+		MASS_CHECK_IF_SPARSE_REQUIREMENT(ElementType);
+
+		return UE::Mass::IsA<FMassFragment>(ElementType)
+			? GetSparseElementInternal(ElementType, EntityIterator.GetEntityHandle())
+			: FConstStructView();
+	}
+
+	MASSENTITY_API bool HasSparseElement(TNotNull<const UScriptStruct*> ElementType, const FEntityIterator& EntityIterator) const;
+	MASSENTITY_API bool DoesChunkContainSparseElement(TNotNull<const UScriptStruct*> ElementType) const;
+
 	/** Sparse chunk related operation */
 	const FMassArchetypeEntityCollection& GetEntityCollection() const { return EntityCollection; }
 
@@ -623,6 +787,9 @@ public:
 	MASSENTITY_API void FlushDeferred();
 
 	void ClearExecutionData();
+
+	void SetCurrentArchetype(FMassArchetypeData& Archetype);
+
 	void SetCurrentArchetypeCompositionDescriptor(const FMassArchetypeCompositionDescriptor& Descriptor)
 	{
 		CurrentArchetypeCompositionBitSet = Descriptor.GetElementsBitSet();
@@ -698,6 +865,9 @@ protected:
 			View.FragmentView.Reset();
 		}
 	}
+
+	MASSENTITY_API bool ExecuteSparseElementsFilterImpl(const FSparseEntityIterator& EntityIt) const;
+	bool ExecuteSparseChunkFilterImpl(const FMassArchetypeChunk& Chunk) const;
 };
 
 //-----------------------------------------------------------------------------
@@ -772,5 +942,18 @@ inline void FMassExecutionContext::ClearFragmentViews(const FMassEntityQuery& Re
 	ClearFragmentViews();
 }
 
-#undef CHECK_IF_VALID
-#undef CHECK_IF_READWRITE
+inline bool FMassExecutionContext::ExecuteSparseElementsFilter(const FSparseEntityIterator& EntityIt) const
+{
+	return (bHasSparseRequirements == false) || ExecuteSparseElementsFilterImpl(EntityIt);
+}
+
+inline int32 FMassExecutionContext::GetChunkSerialModificationNumber() const
+{
+	return ChunkEntityIndex.ChunkSerialNumber;
+}
+
+using FMassSparseEntityIterator = FMassExecutionContext::FSparseEntityIterator;
+
+#undef MASS_CHECK_IF_VALID
+#undef MASS_CHECK_IF_READWRITE
+#undef MASS_CHECK_IF_SPARSE_REQUIREMENT
diff --git a/Engine/Source/Runtime/MassEntity/Public/MassRequirements.h b/Engine/Source/Runtime/MassEntity/Public/MassRequirements.h
--- a/Engine/Source/Runtime/MassEntity/Public/MassRequirements.h
+++ b/Engine/Source/Runtime/MassEntity/Public/MassRequirements.h
@@ -330,6 +330,64 @@ struct FMassFragmentRequirements
 	/** Clears given tags out of all collected requirements, including negative ones */
 	MASSENTITY_API FMassFragmentRequirements& ClearTagRequirements(const FMassTagBitSet& TagsToRemoveBitSet);
 
+	/**
+	 * Adds given sparse element type as a requirement.
+	 */
+	template<UE::Mass::CSparse T>
+	FMassFragmentRequirements& AddSparseRequirement(const EMassFragmentPresence Presence = EMassFragmentPresence::All)
+	{
+		check(Presence != EMassFragmentPresence::MAX);
+		checkf(bInitialized, TEXT("Modifying requirements before initialization is not supported."));
+
+		switch (Presence)
+		{
+		case EMassFragmentPresence::All:
+			RequiredAllSparseElements.Add<T>();
+			break;
+		case EMassFragmentPresence::Any:
+			RequiredAnySparseElements.Add<T>();
+			break;
+		case EMassFragmentPresence::None:
+			RequiredNoneSparseElements.Add<T>();
+			break;
+		case EMassFragmentPresence::Optional:
+			RequiredOptionalSparseElements.Add<T>();
+			break;
+		}
+
+		IncrementChangeCounter();
+		return *this;
+	}
+
+	/**
+	 * Adds given sparse element type as a requirement.
+	 */
+	FMassFragmentRequirements& AddSparseRequirement(TNotNull<const UScriptStruct*> ElementType, const EMassFragmentPresence Presence = EMassFragmentPresence::All)
+	{
+		check(Presence != EMassFragmentPresence::MAX);
+		checkf(bInitialized, TEXT("Modifying requirements before initialization is not supported."));
+		checkf(UE::Mass::IsSparse(ElementType), TEXT("%s doesn't represent a valid sparse element type. Make sure to inherit from FMassSparseFragment or FMassSparseTag."), *ElementType->GetName());
+
+		switch (Presence)
+		{
+		case EMassFragmentPresence::All:
+			RequiredAllSparseElements.Add(*ElementType);
+			break;
+		case EMassFragmentPresence::Any:
+			RequiredAnySparseElements.Add(*ElementType);
+			break;
+		case EMassFragmentPresence::None:
+			RequiredNoneSparseElements.Add(*ElementType);
+			break;
+		case EMassFragmentPresence::Optional:
+			RequiredOptionalSparseElements.Add(*ElementType);
+			break;
+		}
+
+		IncrementChangeCounter();
+		return *this;
+	}
+
 	template<typename T>
 	FMassFragmentRequirements& AddChunkRequirement(const EMassFragmentAccess AccessMode, const EMassFragmentPresence Presence = EMassFragmentPresence::All)
 	{
@@ -628,6 +686,7 @@ struct FMassFragmentRequirements
 	 * The function validates requirements we make for queries. See the FMassFragmentRequirements struct description for details.
 	 * Even though the code of the function is non-trivial the consecutive calls will be essentially free due to the result 
 	 * being cached (note that the caching gets invalidated if the composition changes).
+	 * Note: sparse-elements-only requirements are not valid. Use FMassSparseElementIterator instead
 	 * @return whether this query's requirements follow the rules.
 	 */
 	MASSENTITY_API bool CheckValidity() const;
@@ -653,14 +712,24 @@ struct FMassFragmentRequirements
 	const FMassConstSharedFragmentBitSet& GetRequiredAllConstSharedFragments() const;
 	const FMassConstSharedFragmentBitSet& GetRequiredOptionalConstSharedFragments() const;
 	const FMassConstSharedFragmentBitSet& GetRequiredNoneConstSharedFragments() const;
+	const FMassElementBitSet& GetRequiredAllSparseElements() const;
+	const FMassElementBitSet& GetRequiredAnySparseElements() const;
+	const FMassElementBitSet& GetRequiredOptionalSparseElements() const;
+	const FMassElementBitSet& GetRequiredNoneSparseElements() const;
 
 	bool IsInitialized() const;
 	MASSENTITY_API bool IsEmpty() const;
 	bool HasPositiveRequirements() const;
 	bool HasNegativeRequirements() const;
 	bool HasOptionalRequirements() const;
+	bool HasSparseRequirements() const;
 	bool HasLinkedEntityRequirements() const;
 
+	/**
+	 * Note that DoesArchetypeMatchRequirements functions to not check sparse elements. Sparse elements are not
+	 * a part of immutable composition of a given archetype and can change at runtime, and as such cannot be used
+	 * for filtering archetypes when the results of that filtering get cached (like in FMassEntityQuery)
+	 */
 	MASSENTITY_API bool DoesArchetypeMatchRequirements(const FMassArchetypeHandle& ArchetypeHandle) const;
 	MASSENTITY_API bool DoesArchetypeMatchRequirements(const FMassElementBitSet& ArchetypeCompositionBitSet) const;
 	bool DoesArchetypeMatchRequirements(const FMassArchetypeCompositionDescriptor& ArchetypeComposition) const
@@ -710,6 +779,17 @@ protected:
 	FMassFragmentBitSet RequiredAnyFragments;
 	FMassFragmentBitSet RequiredOptionalFragments;
 	FMassFragmentBitSet RequiredNoneFragments;
+
+	/**
+	 * Sparse elements. These can only be tested at runtime - i.e. the results cannot be cached, since
+	 * as the entities move between archetypes the individual archetypes may or may not contain entities with
+	 * given sparse elements
+	 */
+	FMassElementBitSet RequiredAllSparseElements;
+	FMassElementBitSet RequiredAnySparseElements;
+	FMassElementBitSet RequiredOptionalSparseElements;
+	FMassElementBitSet RequiredNoneSparseElements;
+
 	FMassChunkFragmentBitSet RequiredAllChunkFragments;
 	FMassChunkFragmentBitSet RequiredOptionalChunkFragments;
 	FMassChunkFragmentBitSet RequiredNoneChunkFragments;
@@ -748,6 +828,7 @@ private:
 	 * still need to be satisfied.
 	 */
 	mutable uint16 bHasOptionalRequirements : 1 = false;
+	mutable uint16 bHasSparseRequirements : 1 = false;
 	
 	uint16 bInitialized : 1 = false;
 	uint16 IncrementalChangesCount = 0;
@@ -1020,6 +1101,26 @@ inline const FMassConstSharedFragmentBitSet& FMassFragmentRequirements::GetRequi
 	return RequiredNoneConstSharedFragments; 
 }
 
+inline const FMassElementBitSet& FMassFragmentRequirements::GetRequiredAllSparseElements() const 
+{ 
+	return RequiredAllSparseElements; 
+}
+
+inline const FMassElementBitSet& FMassFragmentRequirements::GetRequiredAnySparseElements() const 
+{ 
+	return RequiredAnySparseElements; 
+}
+
+inline const FMassElementBitSet& FMassFragmentRequirements::GetRequiredOptionalSparseElements() const 
+{ 
+	return RequiredOptionalSparseElements; 
+}
+
+inline const FMassElementBitSet& FMassFragmentRequirements::GetRequiredNoneSparseElements() const 
+{ 
+	return RequiredNoneSparseElements; 
+}
+
 inline bool FMassFragmentRequirements::IsInitialized() const 
 { 
 	return bInitialized; 
@@ -1040,6 +1141,11 @@ inline bool FMassFragmentRequirements::HasOptionalRequirements() const
 	return bHasOptionalRequirements; 
 }
 
+inline bool FMassFragmentRequirements::HasSparseRequirements() const 
+{ 
+	return bHasSparseRequirements; 
+}
+
 inline bool FMassFragmentRequirements::HasLinkedEntityRequirements() const
 { 
 	return !LinkedEntityFragmentRequirements.IsEmpty(); 
diff --git a/Engine/Source/Runtime/MassEntity/Public/MassSparseElementsStorage.h b/Engine/Source/Runtime/MassEntity/Public/MassSparseElementsStorage.h
new file mode 100644
--- /dev/null
+++ b/Engine/Source/Runtime/MassEntity/Public/MassSparseElementsStorage.h
@@ -0,0 +1,481 @@
+// Copyright Epic Games, Inc. All Rights Reserved.
+
+#pragma once
+
+#include "GenericPlatform/GenericPlatformMath.h"
+#include "MassEntityHandle.h"
+#include "MassElement.h"
+#include "StructUtils/StructView.h"
+#include "MassEntityConcepts.h"
+#include "UObject/ObjectKey.h"
+#include "StructUtils/StructArrayView.h"
+
+#define UE_API MASSENTITY_API
+
+/**
+ * Known limitations:
+ * - doesn't support reloading ScriptStruct types
+ */
+
+struct FMassEntityManager;
+
+namespace UE::Mass
+{
+
+struct FSparseElementIterator;
+
+/**
+ * Storage type for hosting instances of sparse elements.
+ *
+ * Sparse elements are optional tags and fragments that can be added to an entity without causing it to
+ * change archetypes. Sparse elements are not part of Archetype's composition - instead, data on which
+ * entity has which elements is stored per archetype chunk. Sparse tags live there (since they don't contain
+ * any data per instance). Sparse fragment instances live in FSparseElementsStorage.
+ *
+ * FSparseElementsStorage provides a way iterate through all instances of a given sparse element type,
+ * using UE::Mass::FSparseElementIterator. To create one call FMassEntityManager::CreateSparseElementsIterator
+ *
+ * Sparse fragment instances are stored in per-type pools (FSparseElementsStorage::FTypePool) which store data in
+ * chunks (FSparseElementsStorage::FPackedStructDataChunk), and each chunk contains uniform array of instances
+ * of the hosted struct type, along with a bitset indicating which elements are valid (i.e. "set").
+ *
+ * Do not use directly. The storage is to be accessed via FMassEntityManager and other means that forward their
+ * calls to the entity manager.
+ */
+struct FSparseElementsStorage
+{	
+private:
+	friend FSparseElementIterator;
+	struct FTypePool;
+
+public:
+	struct FTypeConfig
+	{
+		FTypeConfig(TNotNull<const UScriptStruct*> InElementType)
+			: AllocSizePerElement(InElementType->GetStructureSize())
+			, ElementType(InElementType)
+		{
+		}
+		FTypeConfig(const FTypeConfig&) = default;
+
+		SIZE_T GetElementSize() const
+		{
+			return AllocSizePerElement;
+		}
+
+		const UScriptStruct* GetType() const 
+		{
+			return ElementType;
+		}
+
+		/**
+		 * Converts down to a power of 2 and sets ElementsPerChunk and ChunkIndexShift
+		 * @return actual value used for ElementsPerChunk
+		 */
+		uint32 SetElementsPerChunk(uint32 InElementsPerChunk)
+		{
+			constexpr uint32 MinElementsPerChunk = 1u << 5;
+			constexpr uint32 MaxChunkShift = 32 - 5;
+			InElementsPerChunk = FMath::Max(InElementsPerChunk, MinElementsPerChunk);
+			ChunkIndexShift = FMath::FloorLog2(InElementsPerChunk);
+			check(ChunkIndexShift < MaxChunkShift);
+			ElementsPerChunk = (1u << ChunkIndexShift);
+			ElementIndexMask = ElementsPerChunk - 1;
+			return ElementsPerChunk;
+		}
+
+		uint32 GetChunkIndexShift() const 
+		{ 
+			return ChunkIndexShift; 
+		}
+
+		uint32 GetElementsPerChunk() const 
+		{ 
+			return ElementsPerChunk; 
+		}
+
+	private:
+		friend FTypePool;
+		FTypeConfig();
+		static constexpr int32 DefaultChunkIndexShift = 7;
+		uint32 ChunkIndexShift = DefaultChunkIndexShift;
+		uint32 ElementsPerChunk = 1 << DefaultChunkIndexShift;
+		uint32 ElementIndexMask = (1 << DefaultChunkIndexShift) - 1;
+		const SIZE_T AllocSizePerElement = 0;
+		TNotNull<const UScriptStruct*> ElementType;
+	};
+
+private:
+	struct FPackedStructDataChunk
+	{
+		uint8* RawMemory = nullptr;
+		TBitArray<> OccupationMask;
+		uint32 NumElements = 0;
+
+		FPackedStructDataChunk() = default;
+
+		FPackedStructDataChunk(const FPackedStructDataChunk&) = delete;
+		FPackedStructDataChunk& operator=(const FPackedStructDataChunk&) = delete;
+
+		/** Move constructor required since the array holding FPackedStructDataChunk instances will get reallocated */
+		FPackedStructDataChunk(FPackedStructDataChunk&& Other)
+			: RawMemory(Other.RawMemory)
+			, OccupationMask(MoveTemp(Other.OccupationMask))
+			, NumElements(Other.NumElements)
+		{
+			Other.RawMemory = nullptr;
+			Other.NumElements = 0;
+		}
+
+		FPackedStructDataChunk& operator=(FPackedStructDataChunk&& Other)
+		{
+			if (this != &Other) 
+			{
+				if (RawMemory) 
+				{
+					FMemory::Free(RawMemory);
+				}
+
+				RawMemory = Other.RawMemory;
+				OccupationMask = MoveTemp(Other.OccupationMask);
+				NumElements = Other.NumElements;
+			
+				Other.RawMemory = nullptr;
+				Other.NumElements = 0;
+			}
+			return *this;
+		}
+
+		~FPackedStructDataChunk()
+		{
+			if (RawMemory != nullptr)
+			{
+				FMemory::Free(RawMemory);
+				RawMemory = nullptr;
+			}
+		}
+
+		void Init(const int32 MaxElements, const SIZE_T ElementSize)
+		{
+			OccupationMask.Init(false, MaxElements);
+			RawMemory = static_cast<uint8*>(FMemory::Malloc(MaxElements * ElementSize));
+		}
+
+		bool IsInitialized() const
+		{
+			return RawMemory != nullptr;
+		}
+
+		bool HasElement(const uint32 ElementIndex)
+		{
+			checkSlow(OccupationMask.IsValidIndex(ElementIndex));
+			return OccupationMask[ElementIndex];
+		}
+
+		/** ElementIndex is verified by the caller */
+		uint8* GetElement(const uint32 ElementIndex, const SIZE_T ElementSize)
+		{
+			checkSlow(OccupationMask.IsValidIndex(ElementIndex));
+			return OccupationMask[ElementIndex] ? &RawMemory[ElementSize * ElementIndex] : nullptr;
+		}
+
+		const uint8* GetElement(const uint32 ElementIndex, const SIZE_T ElementSize) const
+		{
+			checkSlow(OccupationMask.IsValidIndex(ElementIndex));
+			return OccupationMask[ElementIndex] ? &RawMemory[ElementSize * ElementIndex] : nullptr;
+		}
+
+		/** returns uninitialized memory, to be initialized by the caller */
+		uint8* Add(const uint32 ElementIndex, const SIZE_T ElementSize)
+		{
+			checkSlow(HasElement(ElementIndex) == false);
+			OccupationMask[ElementIndex] = true;
+			++NumElements;
+			return &RawMemory[ElementSize * ElementIndex];
+		}
+
+		uint8* Remove(const uint32 ElementIndex, const SIZE_T ElementSize)
+		{
+			if (HasElement(ElementIndex))
+			{
+				OccupationMask[ElementIndex] = false;
+				--NumElements;
+				return &RawMemory[ElementSize * ElementIndex];
+			}
+			return nullptr;
+		}
+	private:
+		friend FTypePool;
+		uint8* GetElementUnsafe(const uint32 ElementIndex, const SIZE_T ElementSize) const
+		{
+			return &RawMemory[ElementSize * ElementIndex];
+		}
+	};
+
+	struct FTypePool
+	{
+		FTypePool() = default;
+		FTypePool(const FTypeConfig& InConfig)
+			: Config(InConfig)
+		{
+		}
+		UE_API ~FTypePool();
+
+		TNotNull<uint8*> Add(uint32 EntityIndex);
+		TNotNull<uint8*> Add(uint32 EntityIndex, TNotNull<const uint8*> SourceData);
+		FStructView Add_GetView(uint32 EntityIndex);
+		bool Remove(uint32 EntityIndex);
+		FConstStructView Get(uint32 EntityIndex) const;
+		FStructView GetMutable(uint32 EntityIndex);
+		FStructView GetOrCreateMutable(uint32 EntityIndex);
+
+		uint32 GetNumElements() const
+		{
+			return NumElements;
+		}
+
+		UE_API static FTypePool& GetDummyPool();
+
+	private:
+		friend FSparseElementIterator;
+
+		uint32 GetChunkIndex(const uint32 EntityIndex) const
+		{
+			return EntityIndex >> Config.ChunkIndexShift;
+		}
+
+		uint32 GetElementIndex(const uint32 EntityIndex) const
+		{
+			return EntityIndex & Config.ElementIndexMask;
+		}
+
+		FPackedStructDataChunk* GetPackedData(uint32 ChunkIndex);
+		const FPackedStructDataChunk* GetPackedData(uint32 ChunkIndex) const;
+		FPackedStructDataChunk& GetOrCreatePackedData(uint32 ChunkIndex);
+
+		void GetChunkAndElementIndex(const uint32 EntityIndex, uint32& OutChunkIndex, uint32& OutElementIndex) const
+		{
+			OutChunkIndex = GetChunkIndex(EntityIndex);
+			OutElementIndex = GetElementIndex(EntityIndex);
+		}
+
+		TArray<FPackedStructDataChunk> Chunks;
+		/** stores total number of instances of the given sparse element type */
+		uint32 NumElements = 0;
+		const FTypeConfig Config;
+	};
+
+public:
+	UE_API void Deinitialize();
+
+	/**
+	 * Creates and instance of FSparseElementIterator that can be used to iterate over all available
+	 * instances of sparse fragment of type ElementType
+	 */
+	FSparseElementIterator CreateElementIterator(TNotNull<const UScriptStruct*> ElementType);
+
+	UE_API FStructView AddElementToEntity(FMassEntityHandle EntityHandle, TNotNull<const UScriptStruct*> ElementType);
+
+	template<CSparse T>
+	requires CFragment<T>
+	T& AddElementToEntity(FMassEntityHandle EntityHandle)
+	{
+		return AddElementToEntity(EntityHandle, T::StaticStruct()).template Get<T>();
+	}
+
+	FStructView AddElementInstanceToEntity(FMassEntityHandle EntityHandle, const FConstStructView Instance)
+	{
+		const UScriptStruct* FragmentType = Instance.GetScriptStruct();
+		checkfSlow(IsA<FMassFragment>(FragmentType) && IsSparse(FragmentType), TEXT("%s is not a sparse fragment type"), *GetNameSafe(FragmentType));
+		FStructView ElementInstance = AddElementToEntity(EntityHandle, Instance.GetScriptStruct());
+		FragmentType->CopyScriptStruct(ElementInstance.GetMemory(), Instance.GetMemory());
+		return ElementInstance;
+	}
+
+	UE_API bool RemoveElementFromEntity(FMassEntityHandle EntityHandle, TNotNull<const UScriptStruct*> ElementType);
+
+	template<CSparse T>
+	requires CFragment<T>
+	bool RemoveElementFromEntity(FMassEntityHandle EntityHandle)
+	{
+		return RemoveElementFromEntity(EntityHandle, T::StaticStruct());
+	}
+
+	/**
+	 * @param InOutEntityHandles may get modified by the function - sorted and invalid handles removed 
+	 */
+	void BatchAddElementToEntities(TArray<FMassEntityHandle>& InOutEntityHandles, TNotNull<const UScriptStruct*> ElementType, TArray<FStructView>* OutAddedElements = nullptr);
+	void BatchAddElementToEntities(TArrayView<const FMassEntityHandle> InEntityHandles, TNotNull<const UScriptStruct*> ElementType, TArray<FStructView>* OutAddedElements = nullptr);
+
+	template<CSparse T>
+	requires CFragment<T>
+	void BatchAddElementToEntities(TArray<FMassEntityHandle>& InOutEntityHandles, TArray<FStructView>* OutAddedElements = nullptr)
+	{
+		BatchAddElementToEntities(InOutEntityHandles, T::StaticStruct(), OutAddedElements);
+	}
+
+	void BatchAddElementInstancesToEntities(TArrayView<const FMassEntityHandle> InEntityHandles, const FStructArrayView& FragmentPayload);
+	/**
+	 * Adds FragmentInstance (type and values) to each entity indicated by InEntityHandles
+	 */
+	void BatchAddElementInstancesToEntities(TArrayView<const FMassEntityHandle> InEntityHandles, FConstStructView FragmentInstance);
+
+	/**
+	 * @param InOutEntityHandles may get modified by the function - sorted and invalid handles removed 
+	 * @return num element instances actually removed
+	 */
+	int32 BatchRemoveElementFromEntities(TArray<FMassEntityHandle>& InOutEntityHandles, TNotNull<const UScriptStruct*> ElementType);
+	int32 BatchRemoveElementFromEntities(TArrayView<const FMassEntityHandle> InEntityHandles, TNotNull<const UScriptStruct*> ElementType);
+
+	/** @return num element instances actually removed */
+	template<CSparse T>
+	requires CFragment<T>
+	int32 BatchRemoveElementFromEntities(TArray<FMassEntityHandle>& InOutEntityHandles)
+	{
+		return BatchRemoveElementFromEntities(InOutEntityHandles, T::StaticStruct());
+	}
+
+	/** Clears sparse elements instances belonging to entities indicated by EntityHandles */
+	UE_API void RemoveEntity(FMassEntityHandle EntityHandle, const FMassElementBitSet& Elements);
+
+	/**
+	 * Returns the total number of instances of the given sparse element type
+	 */
+	uint32 GetNumElementsOfType(TNotNull<const UScriptStruct*> ElementType) const
+	{
+		const FTypePool* TypePool = GetTypePool(ElementType);
+		return TypePool ? TypePool->GetNumElements() : 0;
+	}
+
+	template<CSparse T>
+	requires CFragment<T>
+	uint32 GetNumElementsOfType() const
+	{
+		return GetNumElementsOfType(T::StaticStruct());
+	}
+
+	/**
+	 * Checks whether there are any instances of the given sparse element type
+	 */
+	bool HasAnyElementsOfType(TNotNull<const UScriptStruct*> ElementType) const
+	{
+		const FTypePool* TypePool = GetTypePool(ElementType);
+		return TypePool && (TypePool->GetNumElements() > 0);
+	}
+
+	template<CSparse T>
+	requires CFragment<T>
+	bool HasAnyElementsOfType() const
+	{
+		return HasAnyElementsOfType(T::StaticStruct());
+	}
+
+	UE_API FStructView GetMutableElementDataForEntity(const FMassEntityHandle EntityHandle, TNotNull<const UScriptStruct*> ElementType);
+	UE_API FConstStructView GetElementDataForEntity(const FMassEntityHandle EntityHandle, TNotNull<const UScriptStruct*> ElementType) const;
+
+	UE_API void ConfigureType(TNotNull<const UScriptStruct*> ElementType, FTypeConfig& Config);
+
+private:
+	template<typename TContainer>
+	void BatchAddElementToEntities(TContainer InEntityHandles, TNotNull<const UScriptStruct*> ElementType, TArray<FStructView>* OutAddedElements);
+
+	TSparseArray<FTypePool> TypePools;
+
+	/** @param InTypeIndex if INDEX_NONE then the index will be calculated based on ElementType */
+	FTypePool& GetOrCreateTypePool(TNotNull<const UScriptStruct*> ElementType, const int32 InTypeIndex = INDEX_NONE);
+	UE_API const FTypePool* GetTypePool(TNotNull<const UScriptStruct*> ElementType) const;
+	UE_API FTypePool* GetTypePool(TNotNull<const UScriptStruct*> ElementType);
+};
+
+/** To iterate over continuous lists of sparse elements of specific types */
+struct FSparseElementIterator 
+{
+	FSparseElementIterator()
+		: Pool(FSparseElementsStorage::FTypePool::GetDummyPool())
+	{
+	}
+private:
+	friend FSparseElementsStorage;
+	// instances to be created only via FSparseElementsStorage
+	FSparseElementIterator(const FSparseElementsStorage::FTypePool& InPool)
+		: Pool(InPool), ChunkIndex(0)
+	{
+		++(*this);
+	}
+public:
+	bool IsValid() const 
+	{
+		return Pool.Chunks.IsValidIndex(ChunkIndex)
+			&& ElementIndex < Pool.Config.GetElementsPerChunk();
+	}
+
+	explicit operator bool() const
+	{
+		return IsValid();
+	}
+
+	void operator++()
+	{
+		while (ChunkIndex < Pool.Chunks.Num()) 
+		{
+			const FSparseElementsStorage::FPackedStructDataChunk& Chunk = Pool.Chunks[ChunkIndex];
+			if (Chunk.IsInitialized())
+			{
+				checkfSlow(Chunk.OccupationMask.IsEmpty() == false, TEXT("Initialized chunks are expected to have their OccupationMask filled with data."));
+				// Find next bit set
+				TBitArray<>::FConstIterator It(Chunk.OccupationMask, ElementIndex + 1);
+				while (static_cast<uint32>(It.GetIndex()) < Pool.Config.GetElementsPerChunk()) 
+				{
+					if (It.GetValue()) 
+					{
+						ElementIndex = It.GetIndex();
+						return;
+					}
+					++It;
+				}
+			}
+			// Move to next chunk
+			++ChunkIndex;
+			ElementIndex = INDEX_NONE;
+		}
+	}
+
+	FStructView GetElementView() const 
+	{
+		checkSlow(IsValid());
+		const FSparseElementsStorage::FPackedStructDataChunk& Chunk = Pool.Chunks[ChunkIndex];
+		checkSlow(Chunk.OccupationMask.IsValidIndex(ElementIndex) && Chunk.OccupationMask[ElementIndex]);
+		uint8* Memory = &Chunk.RawMemory[ElementIndex * Pool.Config.GetElementSize()];
+		return FStructView(Pool.Config.GetType(), Memory);
+	}
+
+	FStructView operator*() const
+	{
+		return GetElementView();
+	}
+
+	int32 GetEntityIndex() const
+	{
+		return ChunkIndex * Pool.Config.GetElementsPerChunk() + ElementIndex;
+	}
+
+	FMassEntityHandle GetEntityHandle(const FMassEntityManager& EntityManager) const;
+
+private:
+	const FSparseElementsStorage::FTypePool& Pool;
+	int32 ChunkIndex = INDEX_NONE;
+	uint32 ElementIndex = INDEX_NONE;
+};
+
+//-----------------------------------------------------------------------------
+// INLINES
+//-----------------------------------------------------------------------------
+inline FSparseElementIterator FSparseElementsStorage::CreateElementIterator(TNotNull<const UScriptStruct*> ElementType)
+{
+	FTypePool* TypePool = GetTypePool(ElementType);
+	return TypePool ? FSparseElementIterator(*TypePool) : FSparseElementIterator();
+}
+
+}
+
+#undef UE_API
