[MassEntity] fixed unit tests using unregistered shared struct

#Synced-CL 51096110

diff --git a/Engine/Source/Developer/MassEntityTestSuite/Private/MassQueryTest.cpp b/Engine/Source/Developer/MassEntityTestSuite/Private/MassQueryTest.cpp
--- a/Engine/Source/Developer/MassEntityTestSuite/Private/MassQueryTest.cpp
+++ b/Engine/Source/Developer/MassEntityTestSuite/Private/MassQueryTest.cpp
@@ -645,13 +645,13 @@ struct FQueryTest_AllOptional : FEntityTestBase
 
 		{
 			FTestConstSharedFragment_Int FragmentInstance;
-			FConstSharedStruct SharedFragmentInstance = FSharedStruct::Make(FragmentInstance);
+			FConstSharedStruct SharedFragmentInstance = EntityManager->GetOrCreateConstSharedFragment(FragmentInstance);
 			EntityManager->AddConstSharedFragmentToEntity(Entities[CurrentEntityIndex++], SharedFragmentInstance);
 			++ExpectedNumOfArchetypes;
 		}
 		{
 			FTestConstSharedFragment_Float FragmentInstance;
-			FConstSharedStruct SharedFragmentInstance = FSharedStruct::Make(FragmentInstance);
+			FConstSharedStruct SharedFragmentInstance = EntityManager->GetOrCreateConstSharedFragment(FragmentInstance);
 			EntityManager->AddConstSharedFragmentToEntity(Entities[CurrentEntityIndex++], SharedFragmentInstance);
 		}
 		Query.CacheArchetypes();
@@ -845,13 +845,13 @@ struct FQueryTest_JustAConstSharedFragment : FEntityTestBase
 
 		{
 			FTestConstSharedFragment_Int FragmentInstance;
-			FConstSharedStruct SharedFragmentInstance = FSharedStruct::Make(FragmentInstance);
+			FConstSharedStruct SharedFragmentInstance = EntityManager->GetOrCreateConstSharedFragment(FragmentInstance);
 			EntityManager->AddConstSharedFragmentToEntity(Entities[CurrentEntityIndex++], SharedFragmentInstance);
 			++ExpectedNumOfArchetypes;
 		}
 		{
 			FTestConstSharedFragment_Float FragmentInstance;
-			FConstSharedStruct SharedFragmentInstance = FSharedStruct::Make(FragmentInstance);
+			FConstSharedStruct SharedFragmentInstance = EntityManager->GetOrCreateConstSharedFragment(FragmentInstance);
 			EntityManager->AddConstSharedFragmentToEntity(Entities[CurrentEntityIndex++], SharedFragmentInstance);
 		}
 		Query.CacheArchetypes();
@@ -1056,7 +1056,7 @@ struct FQueryTest_LinkedFragmentAccess : FEntityTestBase
 
 		FMassEntityLinkFragment LinkedEntityFragment;
 		LinkedEntityFragment.LinkedEntityHandle = LinkedEntity;
-		FConstSharedStruct SharedLinkedEntityFragment = FSharedStruct::Make(LinkedEntityFragment);
+		FConstSharedStruct SharedLinkedEntityFragment = EntityManager->GetOrCreateConstSharedFragment(LinkedEntityFragment);
 
 		TArray<FMassEntityHandle> Entities;
 		EntityManager->BatchCreateEntities(IntsArchetype, 10, Entities);
diff --git a/Engine/Source/Developer/MassEntityTestSuite/Private/MassSharedFragmentsTest.cpp b/Engine/Source/Developer/MassEntityTestSuite/Private/MassSharedFragmentsTest.cpp
--- a/Engine/Source/Developer/MassEntityTestSuite/Private/MassSharedFragmentsTest.cpp
+++ b/Engine/Source/Developer/MassEntityTestSuite/Private/MassSharedFragmentsTest.cpp
@@ -522,14 +522,24 @@ struct FConstSharedFragment_AddToEntity : FEntityTestBase
 	{
 		constexpr int32 TestIntValue = 1023;
 		const FTestConstSharedFragment_Int FragmentInstance(TestIntValue);
-		const FConstSharedStruct SharedFragmentInstance = FSharedStruct::Make(FragmentInstance);
 
 		const FMassEntityHandle EntityHandle = EntityManager->CreateEntity(FloatsArchetype);
 
 		const FTestConstSharedFragment_Int* EntitySharedFragment = EntityManager->GetConstSharedFragmentDataPtr<FTestConstSharedFragment_Int>(EntityHandle);
 		AITEST_NULL("Initially the entity is not expected to have the shared fragment", EntitySharedFragment);
 
-		EntityManager->AddConstSharedFragmentToEntity(EntityHandle, SharedFragmentInstance);
+		{
+			GetTestRunner().AddExpectedError(TEXT("shared struct requires to be registered first"), EAutomationExpectedErrorFlags::Contains, 1);
+			const FConstSharedStruct UnregisteredSharedFragmentInstance = FSharedStruct::Make(FragmentInstance);
+			const bool bSuccessfullyAddedConstSharedFragment =  EntityManager->AddConstSharedFragmentToEntity(EntityHandle, UnregisteredSharedFragmentInstance);
+			AITEST_FALSE("Adding unregistered const shared fragment struct should fail", bSuccessfullyAddedConstSharedFragment);
+		}
+
+		{
+			const FConstSharedStruct& RegisteredSharedFragmentInstance = EntityManager->GetOrCreateConstSharedFragment(FragmentInstance);
+			const bool bSuccessfullyAddedConstSharedFragment = EntityManager->AddConstSharedFragmentToEntity(EntityHandle, RegisteredSharedFragmentInstance);
+			AITEST_TRUE("Adding registerd const shared fragment struct should succeed", bSuccessfullyAddedConstSharedFragment);
+		}
 
 		EntitySharedFragment = EntityManager->GetConstSharedFragmentDataPtr<FTestConstSharedFragment_Int>(EntityHandle);
 		AITEST_NOT_NULL("The entity is expected to have the shared fragment after the operation", EntitySharedFragment);
@@ -538,16 +548,12 @@ struct FConstSharedFragment_AddToEntity : FEntityTestBase
 
 		// at this point the Entity already has a shared fragment of a given type
 		// now we're going to add it again and test the systems behavior, we'll be adding the same FMasSharedFragment type
-		// in both const and non-const way.
 		constexpr int32 DifferentTestIntValue = TestIntValue + 1;
 		const FTestConstSharedFragment_Int DifferentFragmentInstance(DifferentTestIntValue);
-		const FSharedStruct DifferentSharedFragmentInstance = FSharedStruct::Make(DifferentFragmentInstance);
-		const FConstSharedStruct DifferentConstSharedFragmentInstance = FConstSharedStruct::Make(DifferentFragmentInstance);
+		const FConstSharedStruct DifferentConstSharedFragmentInstance = EntityManager->GetOrCreateConstSharedFragment(DifferentFragmentInstance);
 
-		GetTestRunner().AddExpectedError(TEXT("Changing shared fragment value of entities is not supported"), EAutomationExpectedErrorFlags::Contains, 2);
+		GetTestRunner().AddExpectedError(TEXT("Changing shared fragment value of entities is not supported"), EAutomationExpectedErrorFlags::Contains, 1);
 
-		const bool bSuccessfullyAddedSharedFragment = EntityManager->AddConstSharedFragmentToEntity(EntityHandle, DifferentSharedFragmentInstance);
-		AITEST_FALSE("Adding existing shared fragment type should fail", bSuccessfullyAddedSharedFragment);
 		const bool bSuccessfullyAddedConstSharedFragment = EntityManager->AddConstSharedFragmentToEntity(EntityHandle, DifferentConstSharedFragmentInstance);
 		AITEST_FALSE("Adding existing const shared fragment type should fail", bSuccessfullyAddedConstSharedFragment);
 
@@ -562,7 +568,7 @@ struct FSharedFragment_AddToEntity : FEntityTestBase
 	{
 		constexpr int32 TestIntValue = 1023;
 		const FTestSharedFragment_Int FragmentInstance(TestIntValue);
-		const FSharedStruct SharedFragmentInstance = FSharedStruct::Make(FragmentInstance);
+		const FSharedStruct SharedFragmentInstance = EntityManager->GetOrCreateSharedFragment(FragmentInstance);
 
 		const FMassEntityHandle EntityHandle = EntityManager->CreateEntity(FloatsArchetype);
 
@@ -579,7 +585,7 @@ struct FSharedFragment_AddToEntity : FEntityTestBase
 		// now we're going to add it again and test the systems behavior.
 		constexpr int32 DifferentTestIntValue = TestIntValue + 1;
 		const FTestSharedFragment_Int DifferentFragmentInstance(DifferentTestIntValue);
-		const FSharedStruct DifferentSharedFragmentInstance = FSharedStruct::Make(DifferentFragmentInstance);
+		const FSharedStruct DifferentSharedFragmentInstance = EntityManager->GetOrCreateSharedFragment(DifferentFragmentInstance);
 
 		GetTestRunner().AddExpectedError(TEXT("Changing shared fragment value of entities is not supported"), EAutomationExpectedErrorFlags::Contains, 1);
 
@@ -597,7 +603,7 @@ struct FConstSharedFragment_RemoveFromEntity : FEntityTestBase
 	{
 		constexpr int32 TestIntValue = 1023;
 		const FTestConstSharedFragment_Int FragmentInstance(TestIntValue);
-		const FConstSharedStruct SharedFragmentInstance = FSharedStruct::Make(FragmentInstance);
+		const FConstSharedStruct SharedFragmentInstance = EntityManager->GetOrCreateConstSharedFragment(FragmentInstance);
 
 		const FMassEntityHandle EntityHandle = EntityManager->CreateEntity(FloatsArchetype);
 
@@ -634,7 +640,7 @@ struct FSharedFragment_RemoveFromEntity : FEntityTestBase
 	{
 		constexpr int32 TestIntValue = 1023;
 		const FTestSharedFragment_Int FragmentInstance(TestIntValue);
-		const FSharedStruct SharedFragmentInstance = FSharedStruct::Make(FragmentInstance);
+		const FSharedStruct SharedFragmentInstance = EntityManager->GetOrCreateSharedFragment(FragmentInstance);
 
 		const FMassEntityHandle EntityHandle = EntityManager->CreateEntity(FloatsArchetype);
 
diff --git a/Engine/Source/Runtime/MassEntity/Private/MassEntityManager.cpp b/Engine/Source/Runtime/MassEntity/Private/MassEntityManager.cpp
--- a/Engine/Source/Runtime/MassEntity/Private/MassEntityManager.cpp
+++ b/Engine/Source/Runtime/MassEntity/Private/MassEntityManager.cpp
@@ -1574,8 +1574,9 @@ bool FMassEntityManager::SwapConstSharedFragmentForEntity(const FMassEntityHandl
 		return false;
 	}
 
-	checkf(ConstSharedFragmentsContainer.Find(InConstSharedFragment.GetScriptStruct())
-		, TEXT("Provided shared struct requires to be registered first (i.e., GetGetOrCreateConstSharedFragment"));
+	testableCheckfReturn(ConstSharedFragmentsContainer.Find(InConstSharedFragment.GetScriptStruct()) != nullptr
+		, return false
+		, TEXT("Provided shared struct requires to be registered first (i.e., GetGetOrCreateConstSharedFragment)"));
 
 	CheckIfEntityIsActive(Entity);
 
@@ -1610,11 +1611,12 @@ bool FMassEntityManager::AddConstSharedFragmentToEntity(const FMassEntityHandle
 		return false;
 	}
 
-	checkf(ConstSharedFragmentsContainer.Find(InConstSharedFragment.GetScriptStruct())
-		, TEXT("Provided shared struct requires to be registered first (i.e., GetGetOrCreateConstSharedFragment"));
+	testableCheckfReturn(ConstSharedFragmentsContainer.Find(InConstSharedFragment.GetScriptStruct()) != nullptr
+		, return false
+		, TEXT("Provided shared struct requires to be registered first (i.e., GetGetOrCreateConstSharedFragment)"));
 
 	CheckIfEntityIsActive(Entity);
-	
+
 	FMassArchetypeData* CurrentArchetype = GetEntityStorageInterface().GetArchetypeAsShared(Entity.Index).Get();
 	check(CurrentArchetype);
 
@@ -1712,11 +1714,12 @@ bool FMassEntityManager::AddSharedFragmentToEntity(const FMassEntityHandle Entit
 		return false;
 	}
 
-	checkf(SharedFragmentsContainer.Find(InSharedFragment.GetScriptStruct())
-		, TEXT("Provided shared struct requires to be registered first (i.e., GetGetOrCreateSharedFragment"));
+	testableCheckfReturn(SharedFragmentsContainer.Find(InSharedFragment.GetScriptStruct()) != nullptr
+		, return false
+		, TEXT("Provided shared struct requires to be registered first (i.e., GetGetOrCreateSharedFragment)"));
 
 	CheckIfEntityIsActive(Entity);
-	
+
 	FMassArchetypeData* CurrentArchetype = GetEntityStorageInterface().GetArchetypeAsShared(Entity.Index).Get();
 	check(CurrentArchetype);
 
