[Mass] Fixed FConstSharedFragment_AddToEntity unit test failing to work properly in some scenarios

#rnx

#Synced-CL 51211040

diff --git a/Engine/Source/Developer/MassEntityTestSuite/Private/MassSharedFragmentsTest.cpp b/Engine/Source/Developer/MassEntityTestSuite/Private/MassSharedFragmentsTest.cpp
--- a/Engine/Source/Developer/MassEntityTestSuite/Private/MassSharedFragmentsTest.cpp
+++ b/Engine/Source/Developer/MassEntityTestSuite/Private/MassSharedFragmentsTest.cpp
@@ -529,7 +529,7 @@ struct FConstSharedFragment_AddToEntity : FEntityTestBase
 		AITEST_NULL("Initially the entity is not expected to have the shared fragment", EntitySharedFragment);
 
 		{
-			GetTestRunner().AddExpectedError(TEXT("shared struct requires to be registered first"), EAutomationExpectedErrorFlags::Contains, 1);
+			AITEST_SCOPED_CHECK("shared struct requires to be registered first", 1);
 			const FConstSharedStruct UnregisteredSharedFragmentInstance = FSharedStruct::Make(FragmentInstance);
 			const bool bSuccessfullyAddedConstSharedFragment =  EntityManager->AddConstSharedFragmentToEntity(EntityHandle, UnregisteredSharedFragmentInstance);
 			AITEST_FALSE("Adding unregistered const shared fragment struct should fail", bSuccessfullyAddedConstSharedFragment);
@@ -538,7 +538,7 @@ struct FConstSharedFragment_AddToEntity : FEntityTestBase
 		{
 			const FConstSharedStruct& RegisteredSharedFragmentInstance = EntityManager->GetOrCreateConstSharedFragment(FragmentInstance);
 			const bool bSuccessfullyAddedConstSharedFragment = EntityManager->AddConstSharedFragmentToEntity(EntityHandle, RegisteredSharedFragmentInstance);
-			AITEST_TRUE("Adding registerd const shared fragment struct should succeed", bSuccessfullyAddedConstSharedFragment);
+			AITEST_TRUE("Adding registered const shared fragment struct should succeed", bSuccessfullyAddedConstSharedFragment);
 		}
 
 		EntitySharedFragment = EntityManager->GetConstSharedFragmentDataPtr<FTestConstSharedFragment_Int>(EntityHandle);
@@ -552,7 +552,7 @@ struct FConstSharedFragment_AddToEntity : FEntityTestBase
 		const FTestConstSharedFragment_Int DifferentFragmentInstance(DifferentTestIntValue);
 		const FConstSharedStruct DifferentConstSharedFragmentInstance = EntityManager->GetOrCreateConstSharedFragment(DifferentFragmentInstance);
 
-		GetTestRunner().AddExpectedError(TEXT("Changing shared fragment value of entities is not supported"), EAutomationExpectedErrorFlags::Contains, 1);
+		AITEST_SCOPED_CHECK("Changing shared fragment value of entities is not supported", 1);
 
 		const bool bSuccessfullyAddedConstSharedFragment = EntityManager->AddConstSharedFragmentToEntity(EntityHandle, DifferentConstSharedFragmentInstance);
 		AITEST_FALSE("Adding existing const shared fragment type should fail", bSuccessfullyAddedConstSharedFragment);
