[Mass] Multi-fragment observer support in UMassObserverProcessor

UMassObserverProcessor now has ObservedTypes array so an observer can watch multiple element types without subclassing.
Added deduplicating processor instances in HandleElementsImpl so Execute() is called exactly once per batch even when several observed types are present on the same entities.

Release notes: the change causes a shift in current behavior in use cases where users manually registered a single observer processor instance to observe multiple types, and some of those types were part of a single operations, the instanced observer will get called just once now, whereas in the past it would be called as many times as there were matches with the observed elements.

#Synced-CL 51590771

#Synced-CL 51590771

diff --git a/Engine/Plugins/AI/MassAI/Source/MassAIBehavior/Private/MassLookAtProcessors.cpp b/Engine/Plugins/AI/MassAI/Source/MassAIBehavior/Private/MassLookAtProcessors.cpp
--- a/Engine/Plugins/AI/MassAI/Source/MassAIBehavior/Private/MassLookAtProcessors.cpp
+++ b/Engine/Plugins/AI/MassAI/Source/MassAIBehavior/Private/MassLookAtProcessors.cpp
@@ -668,7 +668,7 @@ void UMassLookAtTargetGridProcessor::Execute(FMassEntityManager& EntityManager,
 UMassLookAtTargetRemoverProcessor::UMassLookAtTargetRemoverProcessor()
 	: Query(*this)
 {
-	ObservedType = FMassLookAtTargetFragment::StaticStruct();
+	ObservedTypes.Add(FMassLookAtTargetFragment::StaticStruct());
 	ObservedOperations = EMassObservedOperationFlags::Remove;
 	ExecutionFlags = static_cast<int32>(EProcessorExecutionFlags::AllNetModes);
 }
@@ -699,7 +699,7 @@ void UMassLookAtTargetRemoverProcessor::Execute(FMassEntityManager& EntityManage
 UMassLookAtRequestInitializer::UMassLookAtRequestInitializer()
 	: Query(*this)
 {
-	ObservedType = FMassLookAtRequestFragment::StaticStruct();
+	ObservedTypes.Add(FMassLookAtRequestFragment::StaticStruct());
 	ObservedOperations = EMassObservedOperationFlags::Add;
 	ExecutionFlags = static_cast<int32>(EProcessorExecutionFlags::Client | EProcessorExecutionFlags::Standalone);
 }
@@ -735,7 +735,7 @@ void UMassLookAtRequestInitializer::Execute(FMassEntityManager& EntityManager, F
 UMassLookAtRequestDeinitializer::UMassLookAtRequestDeinitializer()
 	: Query(*this)
 {
-	ObservedType = FMassLookAtRequestFragment::StaticStruct();
+	ObservedTypes.Add(FMassLookAtRequestFragment::StaticStruct());
 	ObservedOperations = EMassObservedOperationFlags::Remove;
 	ExecutionFlags = static_cast<int32>(EProcessorExecutionFlags::Client | EProcessorExecutionFlags::Standalone);
 }
diff --git a/Engine/Plugins/AI/MassAI/Source/MassAIBehavior/Private/MassStateTreeProcessors.cpp b/Engine/Plugins/AI/MassAI/Source/MassAIBehavior/Private/MassStateTreeProcessors.cpp
--- a/Engine/Plugins/AI/MassAI/Source/MassAIBehavior/Private/MassStateTreeProcessors.cpp
+++ b/Engine/Plugins/AI/MassAI/Source/MassAIBehavior/Private/MassStateTreeProcessors.cpp
@@ -104,7 +104,7 @@ UMassStateTreeFragmentDestructor::UMassStateTreeFragmentDestructor()
 	: EntityQuery(*this)
 {
 	ExecutionFlags = static_cast<uint8>(UE::MassStateTree::ExecutionFlags);
-	ObservedType = FMassStateTreeInstanceFragment::StaticStruct();
+	ObservedTypes.Add(FMassStateTreeInstanceFragment::StaticStruct());
 	ObservedOperations = EMassObservedOperationFlags::Remove;
 	bRequiresGameThreadExecution = true;
 }
diff --git a/Engine/Plugins/AI/MassAI/Source/MassAIBehavior/Private/MassZoneGraphAnnotationProcessors.cpp b/Engine/Plugins/AI/MassAI/Source/MassAIBehavior/Private/MassZoneGraphAnnotationProcessors.cpp
--- a/Engine/Plugins/AI/MassAI/Source/MassAIBehavior/Private/MassZoneGraphAnnotationProcessors.cpp
+++ b/Engine/Plugins/AI/MassAI/Source/MassAIBehavior/Private/MassZoneGraphAnnotationProcessors.cpp
@@ -20,7 +20,7 @@
 UMassZoneGraphAnnotationTagsInitializer::UMassZoneGraphAnnotationTagsInitializer()
 	: EntityQuery(*this)
 {
-	ObservedType = FMassZoneGraphAnnotationFragment::StaticStruct();
+	ObservedTypes.Add(FMassZoneGraphAnnotationFragment::StaticStruct());
 	ObservedOperations = EMassObservedOperationFlags::Add;
 }
 
diff --git a/Engine/Plugins/AI/MassAI/Source/MassNavigation/Private/MassNavigationProcessors.cpp b/Engine/Plugins/AI/MassAI/Source/MassNavigation/Private/MassNavigationProcessors.cpp
--- a/Engine/Plugins/AI/MassAI/Source/MassNavigation/Private/MassNavigationProcessors.cpp
+++ b/Engine/Plugins/AI/MassAI/Source/MassNavigation/Private/MassNavigationProcessors.cpp
@@ -123,7 +123,7 @@ void UMassNavigationSmoothHeightProcessor::Execute(FMassEntityManager& EntityMan
 UMassMoveTargetFragmentInitializer::UMassMoveTargetFragmentInitializer()
 	: InitializerQuery(*this)
 {
-	ObservedType = FMassMoveTargetFragment::StaticStruct();
+	ObservedTypes.Add(FMassMoveTargetFragment::StaticStruct());
 	ObservedOperations = EMassObservedOperationFlags::Add;
 }
 
@@ -275,7 +275,7 @@ void UMassNavigationObstacleGridProcessor::Execute(FMassEntityManager& EntityMan
 UMassNavigationObstacleRemoverProcessor::UMassNavigationObstacleRemoverProcessor()
 	: EntityQuery(*this)
 {
-	ObservedType = FMassNavigationObstacleGridCellLocationFragment::StaticStruct();
+	ObservedTypes.Add(FMassNavigationObstacleGridCellLocationFragment::StaticStruct());
 	ObservedOperations = EMassObservedOperationFlags::Remove;
 	ExecutionFlags = (int32)EProcessorExecutionFlags::AllNetModes;
 }
diff --git a/Engine/Plugins/AI/MassAI/Source/MassZoneGraphNavigation/Private/MassZoneGraphNavigationProcessors.cpp b/Engine/Plugins/AI/MassAI/Source/MassZoneGraphNavigation/Private/MassZoneGraphNavigationProcessors.cpp
--- a/Engine/Plugins/AI/MassAI/Source/MassZoneGraphNavigation/Private/MassZoneGraphNavigationProcessors.cpp
+++ b/Engine/Plugins/AI/MassAI/Source/MassZoneGraphNavigation/Private/MassZoneGraphNavigationProcessors.cpp
@@ -28,7 +28,7 @@
 UMassZoneGraphLocationInitializer::UMassZoneGraphLocationInitializer()
 	: EntityQuery(*this)
 {
-	ObservedType = FMassZoneGraphLaneLocationFragment::StaticStruct();
+	ObservedTypes.Add(FMassZoneGraphLaneLocationFragment::StaticStruct());
 	ObservedOperations = EMassObservedOperationFlags::Add;
 	ExecutionFlags = (int32)EProcessorExecutionFlags::AllNetModes;
 }
diff --git a/Engine/Plugins/AI/MassCrowd/Source/MassCrowd/Private/MassCrowdNavigationProcessor.cpp b/Engine/Plugins/AI/MassCrowd/Source/MassCrowd/Private/MassCrowdNavigationProcessor.cpp
--- a/Engine/Plugins/AI/MassCrowd/Source/MassCrowd/Private/MassCrowdNavigationProcessor.cpp
+++ b/Engine/Plugins/AI/MassCrowd/Source/MassCrowd/Private/MassCrowdNavigationProcessor.cpp
@@ -69,7 +69,7 @@ UMassCrowdLaneTrackingDestructor::UMassCrowdLaneTrackingDestructor()
 	: EntityQuery(*this)
 {
 	ExecutionFlags = (int32)(EProcessorExecutionFlags::Standalone | EProcessorExecutionFlags::Server);
-	ObservedType = FMassCrowdLaneTrackingFragment::StaticStruct();
+	ObservedTypes.Add(FMassCrowdLaneTrackingFragment::StaticStruct());
 	ObservedOperations = EMassObservedOperationFlags::Remove;
 }
 
@@ -228,7 +228,7 @@ UMassCrowdDynamicObstacleInitializer::UMassCrowdDynamicObstacleInitializer()
 	: EntityQuery(*this)
 {
 	ExecutionFlags = (int32)(EProcessorExecutionFlags::Standalone | EProcessorExecutionFlags::Server);
-	ObservedType = FMassCrowdObstacleFragment::StaticStruct();
+	ObservedTypes.Add(FMassCrowdObstacleFragment::StaticStruct());
 	ObservedOperations = EMassObservedOperationFlags::Add;
 }
 
@@ -267,7 +267,7 @@ UMassCrowdDynamicObstacleDeinitializer::UMassCrowdDynamicObstacleDeinitializer()
 	: EntityQuery(*this)
 {
 	ExecutionFlags = (int32)(EProcessorExecutionFlags::Standalone | EProcessorExecutionFlags::Server);
-	ObservedType = FMassCrowdObstacleFragment::StaticStruct();
+	ObservedTypes.Add(FMassCrowdObstacleFragment::StaticStruct());
 	ObservedOperations = EMassObservedOperationFlags::Remove;
 }
 
diff --git a/Engine/Plugins/Experimental/EditorDataStorage/Source/TedsCore/Private/Processors/TypedElementProcessorAdaptors.cpp b/Engine/Plugins/Experimental/EditorDataStorage/Source/TedsCore/Private/Processors/TypedElementProcessorAdaptors.cpp
--- a/Engine/Plugins/Experimental/EditorDataStorage/Source/TedsCore/Private/Processors/TypedElementProcessorAdaptors.cpp
+++ b/Engine/Plugins/Experimental/EditorDataStorage/Source/TedsCore/Private/Processors/TypedElementProcessorAdaptors.cpp
@@ -1950,7 +1950,7 @@ FMassEntityQuery& UTypedElementQueryObserverCallbackAdapterProcessorBase::GetQue
 
 const UScriptStruct* UTypedElementQueryObserverCallbackAdapterProcessorBase::GetObservedType() const
 {
-	return ObservedType;
+	return GetSingleObservedType();
 }
 
 EMassObservedOperationFlags UTypedElementQueryObserverCallbackAdapterProcessorBase::GetObservedOperations() const
@@ -1979,8 +1979,9 @@ bool UTypedElementQueryObserverCallbackAdapterProcessorBase::ConfigureQueryCallb
 
 	bRequiresGameThreadExecution = Query.Description.Callback.ExecutionMode == EExecutionMode::GameThread;
 	ExecutionFlags = static_cast<int32>(EProcessorExecutionFlags::Editor);
-	
-	ObservedType = const_cast<UScriptStruct*>(Query.Description.Callback.MonitoredType);
+
+	ObservedTypes.Reset();
+	ObservedTypes.Add(Query.Description.Callback.MonitoredType);
 	
 	switch (Query.Description.Callback.Type)
 	{
diff --git a/Engine/Plugins/Experimental/UAF/UAFMass/Source/UAFMass/Private/Processors/MassUAFProcessor.cpp b/Engine/Plugins/Experimental/UAF/UAFMass/Source/UAFMass/Private/Processors/MassUAFProcessor.cpp
--- a/Engine/Plugins/Experimental/UAF/UAFMass/Source/UAFMass/Private/Processors/MassUAFProcessor.cpp
+++ b/Engine/Plugins/Experimental/UAF/UAFMass/Source/UAFMass/Private/Processors/MassUAFProcessor.cpp
@@ -27,7 +27,7 @@ namespace UE::UAF::Mass
 UMassUAFInitializer::UMassUAFInitializer()
 	: EntityQuery(*this)
 {
-	ObservedType = FMassUAFFragment::StaticStruct();
+	ObservedTypes.Add(FMassUAFFragment::StaticStruct());
 	ObservedOperations = EMassObservedOperationFlags::Add;
 	ExecutionFlags = (int32)EProcessorExecutionFlags::AllNetModes;
 	bRequiresGameThreadExecution = false;
@@ -143,7 +143,7 @@ void UMassAnimPoseDebugProcessor::Execute(FMassEntityManager& EntityManager, FMa
 UMassUAFDestructor::UMassUAFDestructor()
 	: EntityQuery(*this)
 {
-	ObservedType = FMassUAFFragment::StaticStruct();
+	ObservedTypes.Add(FMassUAFFragment::StaticStruct());
 	ObservedOperations = EMassObservedOperationFlags::Remove;
 	ExecutionFlags = (int32)EProcessorExecutionFlags::AllNetModes;
 	bRequiresGameThreadExecution = false;
diff --git a/Engine/Plugins/Runtime/MassGameplay/Source/MassGameplayDebug/Private/AssignDebugVisProcessor.cpp b/Engine/Plugins/Runtime/MassGameplay/Source/MassGameplayDebug/Private/AssignDebugVisProcessor.cpp
--- a/Engine/Plugins/Runtime/MassGameplay/Source/MassGameplayDebug/Private/AssignDebugVisProcessor.cpp
+++ b/Engine/Plugins/Runtime/MassGameplay/Source/MassGameplayDebug/Private/AssignDebugVisProcessor.cpp
@@ -26,7 +26,7 @@ UAssignDebugVisProcessor::UAssignDebugVisProcessor()
 
 	// Always set the observed type and operations even if WITH_MASSGAMEPLAY_DEBUG is 0
 	// since the processor may still be registered even if it won't do anything.
-	ObservedType = FSimDebugVisFragment::StaticStruct();
+	ObservedTypes.Add(FSimDebugVisFragment::StaticStruct());
 	ObservedOperations = EMassObservedOperationFlags::Add;
 
 #if WITH_MASSGAMEPLAY_DEBUG
diff --git a/Engine/Plugins/Runtime/MassGameplay/Source/MassMovement/Private/MassVelocityRandomizerTrait.cpp b/Engine/Plugins/Runtime/MassGameplay/Source/MassMovement/Private/MassVelocityRandomizerTrait.cpp
--- a/Engine/Plugins/Runtime/MassGameplay/Source/MassMovement/Private/MassVelocityRandomizerTrait.cpp
+++ b/Engine/Plugins/Runtime/MassGameplay/Source/MassMovement/Private/MassVelocityRandomizerTrait.cpp
@@ -28,7 +28,7 @@ void UMassVelocityRandomizerTrait::BuildTemplate(FMassEntityTemplateBuildContext
 UMassRandomVelocityInitializer::UMassRandomVelocityInitializer()
 	: EntityQuery(*this)
 {
-	ObservedType = FMassVelocityFragment::StaticStruct();
+	ObservedTypes.Add(FMassVelocityFragment::StaticStruct());
 	ObservedOperations = EMassObservedOperationFlags::Add;
 	RandomStream.Initialize(UE::Mass::Utils::GenerateRandomSeed());
 }
diff --git a/Engine/Plugins/Runtime/MassGameplay/Source/MassReplication/Private/MassReplicationFragments.cpp b/Engine/Plugins/Runtime/MassGameplay/Source/MassReplication/Private/MassReplicationFragments.cpp
--- a/Engine/Plugins/Runtime/MassGameplay/Source/MassReplication/Private/MassReplicationFragments.cpp
+++ b/Engine/Plugins/Runtime/MassGameplay/Source/MassReplication/Private/MassReplicationFragments.cpp
@@ -16,7 +16,7 @@ UMassNetworkIDFragmentInitializer::UMassNetworkIDFragmentInitializer()
 	: EntityQuery(*this)
 {
 	ExecutionFlags = (int32)(EProcessorExecutionFlags::Standalone | EProcessorExecutionFlags::Server);
-	ObservedType = FMassNetworkIDFragment::StaticStruct();
+	ObservedTypes.Add(FMassNetworkIDFragment::StaticStruct());
 	ObservedOperations = EMassObservedOperationFlags::Add;
 }
 
diff --git a/Engine/Plugins/Runtime/MassGameplay/Source/MassReplication/Private/MassReplicationGridProcessor.cpp b/Engine/Plugins/Runtime/MassGameplay/Source/MassReplication/Private/MassReplicationGridProcessor.cpp
--- a/Engine/Plugins/Runtime/MassGameplay/Source/MassReplication/Private/MassReplicationGridProcessor.cpp
+++ b/Engine/Plugins/Runtime/MassGameplay/Source/MassReplication/Private/MassReplicationGridProcessor.cpp
@@ -132,7 +132,7 @@ void UMassReplicationGridProcessor::Execute(FMassEntityManager& EntityManager, F
 UMassReplicationGridRemoverProcessor::UMassReplicationGridRemoverProcessor()
 	: EntityQuery(*this)
 {
-	ObservedType = FMassReplicationGridCellLocationFragment::StaticStruct();
+	ObservedTypes.Add(FMassReplicationGridCellLocationFragment::StaticStruct());
 	ObservedOperations = EMassObservedOperationFlags::Remove;
 	ExecutionFlags = (int32)EProcessorExecutionFlags::AllNetModes;
 }
diff --git a/Engine/Plugins/Runtime/MassGameplay/Source/MassReplication/Private/MassReplicationProcessor.cpp b/Engine/Plugins/Runtime/MassGameplay/Source/MassReplication/Private/MassReplicationProcessor.cpp
--- a/Engine/Plugins/Runtime/MassGameplay/Source/MassReplication/Private/MassReplicationProcessor.cpp
+++ b/Engine/Plugins/Runtime/MassGameplay/Source/MassReplication/Private/MassReplicationProcessor.cpp
@@ -382,7 +382,7 @@ UMassReplicationEntityDestructionObserver::UMassReplicationEntityDestructionObse
 	: EntityQuery(*this)
 {
 	ExecutionFlags = static_cast<uint8>(EProcessorExecutionFlags::Server);
-	ObservedType = FMassReplicatedAgentFragment::StaticStruct();
+	ObservedTypes.Add(FMassReplicatedAgentFragment::StaticStruct());
 	ObservedOperations = EMassObservedOperationFlags::DestroyEntity;
 
 	// Make sure we can only run in Game Thread to sync properly with UMassReplicationProcessor
diff --git a/Engine/Plugins/Runtime/MassGameplay/Source/MassRepresentation/Private/MassRepresentationProcessor.cpp b/Engine/Plugins/Runtime/MassGameplay/Source/MassRepresentation/Private/MassRepresentationProcessor.cpp
--- a/Engine/Plugins/Runtime/MassGameplay/Source/MassRepresentation/Private/MassRepresentationProcessor.cpp
+++ b/Engine/Plugins/Runtime/MassGameplay/Source/MassRepresentation/Private/MassRepresentationProcessor.cpp
@@ -455,7 +455,7 @@ void UMassVisualizationProcessor::Execute(FMassEntityManager& InEntityManager, F
 UMassRepresentationFragmentDestructor::UMassRepresentationFragmentDestructor()
 	: EntityQuery(*this)
 {
-	ObservedType = FMassRepresentationFragment::StaticStruct();
+	ObservedTypes.Add(FMassRepresentationFragment::StaticStruct());
 	ObservedOperations = EMassObservedOperationFlags::Remove;
 	ExecutionFlags = (int32)EProcessorExecutionFlags::AllNetModes;
 	bRequiresGameThreadExecution = true; // due to FMassRepresentationSubsystemSharedFragment.RepresentationSubsystem use
diff --git a/Engine/Plugins/Runtime/MassGameplay/Source/MassRepresentation/Private/MassStationaryISMRepresentationFragmentDestructor.cpp b/Engine/Plugins/Runtime/MassGameplay/Source/MassRepresentation/Private/MassStationaryISMRepresentationFragmentDestructor.cpp
--- a/Engine/Plugins/Runtime/MassGameplay/Source/MassRepresentation/Private/MassStationaryISMRepresentationFragmentDestructor.cpp
+++ b/Engine/Plugins/Runtime/MassGameplay/Source/MassRepresentation/Private/MassStationaryISMRepresentationFragmentDestructor.cpp
@@ -13,7 +13,7 @@
 UMassStationaryISMRepresentationFragmentDestructor::UMassStationaryISMRepresentationFragmentDestructor()
 	: EntityQuery(*this)
 {
-	ObservedType = FMassRepresentationFragment::StaticStruct();
+	ObservedTypes.Add(FMassRepresentationFragment::StaticStruct());
 	ObservedOperations = EMassObservedOperationFlags::Remove;
 	ExecutionFlags = (int32)EProcessorExecutionFlags::AllWorldModes;
 	bRequiresGameThreadExecution = true; // not sure about this
diff --git a/Engine/Plugins/Runtime/MassGameplay/Source/MassSmartObjects/Private/MassSmartObjectProcessor.cpp b/Engine/Plugins/Runtime/MassGameplay/Source/MassSmartObjects/Private/MassSmartObjectProcessor.cpp
--- a/Engine/Plugins/Runtime/MassGameplay/Source/MassSmartObjects/Private/MassSmartObjectProcessor.cpp
+++ b/Engine/Plugins/Runtime/MassGameplay/Source/MassSmartObjects/Private/MassSmartObjectProcessor.cpp
@@ -425,7 +425,7 @@ void UMassSmartObjectTimedBehaviorProcessor::Execute(FMassEntityManager& EntityM
 UMassSmartObjectUserFragmentDeinitializer::UMassSmartObjectUserFragmentDeinitializer()
 	: EntityQuery(*this)
 {
-	ObservedType = FMassSmartObjectUserFragment::StaticStruct();
+	ObservedTypes.Add(FMassSmartObjectUserFragment::StaticStruct());
 	ObservedOperations = EMassObservedOperationFlags::Remove;
 	ExecutionFlags = (int32)EProcessorExecutionFlags::AllNetModes;
 }
diff --git a/Engine/Plugins/Runtime/MassGameplay/Source/MassSmartObjects/Private/MassSmartObjectRegistration.cpp b/Engine/Plugins/Runtime/MassGameplay/Source/MassSmartObjects/Private/MassSmartObjectRegistration.cpp
--- a/Engine/Plugins/Runtime/MassGameplay/Source/MassSmartObjects/Private/MassSmartObjectRegistration.cpp
+++ b/Engine/Plugins/Runtime/MassGameplay/Source/MassSmartObjects/Private/MassSmartObjectRegistration.cpp
@@ -115,7 +115,7 @@ void UMassSmartObjectDeinitializerBase::Execute(FMassEntityManager& EntityManage
 //----------------------------------------------------------------------//
 UMassActiveSmartObjectInitializer::UMassActiveSmartObjectInitializer()
 {
-	ObservedType = FMassInActiveSmartObjectsRangeTag::StaticStruct();
+	ObservedTypes.Add(FMassInActiveSmartObjectsRangeTag::StaticStruct());
 	Signal = UE::Mass::Signals::SmartObjectActivationChanged;
 }
 
@@ -124,7 +124,7 @@ UMassActiveSmartObjectInitializer::UMassActiveSmartObjectInitializer()
 //----------------------------------------------------------------------//
 UMassActiveSmartObjectDeinitializer::UMassActiveSmartObjectDeinitializer()
 {
-	ObservedType = FMassInActiveSmartObjectsRangeTag::StaticStruct();
+	ObservedTypes.Add(FMassInActiveSmartObjectsRangeTag::StaticStruct());
 	Signal = UE::Mass::Signals::SmartObjectActivationChanged;
 }
 
@@ -133,7 +133,7 @@ UMassActiveSmartObjectDeinitializer::UMassActiveSmartObjectDeinitializer()
 //----------------------------------------------------------------------//
 UMassActorInstanceHandleInitializer::UMassActorInstanceHandleInitializer()
 {
-	ObservedType = FMassActorInstanceFragment::StaticStruct();
+	ObservedTypes.Add(FMassActorInstanceFragment::StaticStruct());
 	Signal = UE::Mass::Signals::ActorInstanceHandleChanged;
 }
 
@@ -142,7 +142,7 @@ UMassActorInstanceHandleInitializer::UMassActorInstanceHandleInitializer()
 //----------------------------------------------------------------------//
 UMassActorInstanceHandleDeinitializer::UMassActorInstanceHandleDeinitializer()
 {
-	ObservedType = FMassActorInstanceFragment::StaticStruct();
+	ObservedTypes.Add(FMassActorInstanceFragment::StaticStruct());
 	Signal = UE::Mass::Signals::ActorInstanceHandleChanged;
 }
 
diff --git a/Engine/Source/Developer/MassEntityTestSuite/Private/MassObserverTest.cpp b/Engine/Source/Developer/MassEntityTestSuite/Private/MassObserverTest.cpp
--- a/Engine/Source/Developer/MassEntityTestSuite/Private/MassObserverTest.cpp
+++ b/Engine/Source/Developer/MassEntityTestSuite/Private/MassObserverTest.cpp
@@ -1022,6 +1022,204 @@ IMPLEMENT_AI_INSTANT_TEST(FMoveToAnotherArchetype_SingleEntity, "System.Mass.Obs
 //};
 //IMPLEMENT_AI_INSTANT_TEST(FMoveToAnotherArchetype_SingleEntity, "System.Mass.Observer.Recursive");
 
+//----------------------------------------------------------------------//
+// Multi-type observer processor tests
+//----------------------------------------------------------------------//
+
+/**
+ * Registers the same UMassTestProcessorBase instance for two fragment types (Float Add and Int Add).
+ * Moving an entity from EmptyArchetype to FloatsIntsArchetype adds both fragments in a single composition
+ * change, which produces one HandleElementsImpl call with both types in ObservedTypesOverlap.
+ * The deduplication set in HandleElementsImpl must prevent Execute from firing twice.
+ */
+struct FCompositeObserver_DeduplicationTwoTypes : FEntityTestBase
+{
+	int32 ExecuteCallCount = 0;
+	UMassTestProcessorBase* ObserverProcessor = nullptr;
+
+	virtual bool SetUp() override
+	{
+		if (!FEntityTestBase::SetUp())
+		{
+			return false;
+		}
+		ObserverProcessor = NewTestProcessor<UMassTestProcessorBase>(EntityManager);
+		// Override ExecutionFunction so we count raw Execute() invocations, not per-chunk calls.
+		ObserverProcessor->ExecutionFunction = [ExecuteCountPtr = &ExecuteCallCount](FMassEntityManager&, FMassExecutionContext&)
+		{
+			++(*ExecuteCountPtr);
+		};
+		return true;
+	}
+
+	virtual bool InstantTest() override
+	{
+		FMassObserverManager& ObserverManager = EntityManager->GetObserverManager();
+		// Register the same processor instance for both Float and Int addition.
+		ObserverManager.AddObserverInstance(FTestFragment_Float::StaticStruct(), EMassObservedOperationFlags::Add, ObserverProcessor);
+		ObserverManager.AddObserverInstance(FTestFragment_Int::StaticStruct(), EMassObservedOperationFlags::Add, ObserverProcessor);
+
+		// Moving EmptyArchetype -> FloatsIntsArchetype adds both Float and Int simultaneously.
+		const FMassEntityHandle Entity = EntityManager->CreateEntity(EmptyArchetype);
+		EntityManager->MoveEntityToAnotherArchetype(Entity, FloatsIntsArchetype);
+
+		AITEST_EQUAL(TEXT("Observer Execute should fire exactly once even though two observed types were added simultaneously"), ExecuteCallCount, 1);
+		return true;
+	}
+};
+IMPLEMENT_AI_INSTANT_TEST(FCompositeObserver_DeduplicationTwoTypes, "System.Mass.Observer.CompositeObserver.DeduplicationTwoTypes");
+
+/**
+ * Registers the same processor for Float Add and Int Add, but the composition change only touches Float
+ * (entity already owns Int). Verifies that Execute fires exactly once — once for the Float pipeline,
+ * and not additionally for Int (which is absent from ObservedTypesOverlap).
+ */
+struct FCompositeObserver_PartialTypeMatch : FEntityTestBase
+{
+	int32 ExecuteCallCount = 0;
+	UMassTestProcessorBase* ObserverProcessor = nullptr;
+
+	virtual bool SetUp() override
+	{
+		if (!FEntityTestBase::SetUp())
+		{
+			return false;
+		}
+		ObserverProcessor = NewTestProcessor<UMassTestProcessorBase>(EntityManager);
+		ObserverProcessor->ExecutionFunction = [ExecuteCountPtr = &ExecuteCallCount](FMassEntityManager&, FMassExecutionContext&)
+		{
+			++(*ExecuteCountPtr);
+		};
+		return true;
+	}
+
+	virtual bool InstantTest() override
+	{
+		FMassObserverManager& ObserverManager = EntityManager->GetObserverManager();
+		ObserverManager.AddObserverInstance(FTestFragment_Float::StaticStruct(), EMassObservedOperationFlags::Add, ObserverProcessor);
+		ObserverManager.AddObserverInstance(FTestFragment_Int::StaticStruct(), EMassObservedOperationFlags::Add, ObserverProcessor);
+
+		// Entity starts in IntsArchetype (already has Int). Moving to FloatsIntsArchetype only adds Float,
+		// so only the Float pipeline is visited in HandleElementsImpl.
+		TArray<FMassEntityHandle> Entities;
+		EntityManager->BatchCreateEntities(IntsArchetype, 1, Entities);
+		AITEST_EQUAL(TEXT("Observer Execute should fire exactly once, for the initial type"), ExecuteCallCount, 1);
+		ExecuteCallCount = 0;
+
+		EntityManager->MoveEntityToAnotherArchetype(Entities[0], FloatsIntsArchetype);
+
+		AITEST_EQUAL(TEXT("Observer Execute should fire exactly once, for the single new type"), ExecuteCallCount, 1);
+		return true;
+	}
+};
+IMPLEMENT_AI_INSTANT_TEST(FCompositeObserver_PartialTypeMatch, "System.Mass.Observer.CompositeObserver.PartialTypeMatch");
+
+/**
+ * Creates a UMassTestCompositeObserverProcessor, sets ObservedTypes = {Float, Int} and
+ * ObservedOperations = Add, then registers it via the generalized
+ * AddObserverInstance(UMassObserverProcessor*) API. Verifies that the processor ends up in
+ * both type pipelines and that Execute fires exactly once when both types are added simultaneously.
+ */
+struct FCompositeObserver_BaseApiAddObserver : FEntityTestBase
+{
+	UMassTestCompositeObserverProcessor* CompositeProcessor = nullptr;
+
+	virtual bool SetUp() override
+	{
+		if (!FEntityTestBase::SetUp())
+		{
+			return false;
+		}
+		CompositeProcessor = NewTestProcessor<UMassTestCompositeObserverProcessor>(EntityManager);
+		CompositeProcessor->SetObservedTypesForTest({ FTestFragment_Float::StaticStruct(), FTestFragment_Int::StaticStruct() });
+		CompositeProcessor->SetObservedOperationsForTest(EMassObservedOperationFlags::Add);
+		return true;
+	}
+
+	virtual bool InstantTest() override
+	{
+		FMassObserverManager& ObserverManager = EntityManager->GetObserverManager();
+		// Register via the generalized base class overload.
+		ObserverManager.AddObserverInstance(CompositeProcessor);
+
+		// Adding both Float and Int simultaneously should trigger exactly one Execute.
+		const FMassEntityHandle Entity = EntityManager->CreateEntity(EmptyArchetype);
+		EntityManager->MoveEntityToAnotherArchetype(Entity, FloatsIntsArchetype);
+
+		AITEST_EQUAL(TEXT("Composite observer registered via base API should execute exactly once"), CompositeProcessor->ExecuteCallCount, 1);
+		return true;
+	}
+};
+IMPLEMENT_AI_INSTANT_TEST(FCompositeObserver_BaseApiAddObserver, "System.Mass.Observer.CompositeObserver.BaseApiAddObserver");
+
+/**
+ * Verifies that RemoveObserverInstance(UMassObserverProcessor*) unregisters the multi-type processor
+ * from all type pipelines it was registered in, so that subsequent composition changes no longer
+ * trigger Execute.
+ */
+struct FCompositeObserver_BaseApiRemoveObserver : FEntityTestBase
+{
+	UMassTestCompositeObserverProcessor* CompositeProcessor = nullptr;
+
+	virtual bool SetUp() override
+	{
+		if (!FEntityTestBase::SetUp())
+		{
+			return false;
+		}
+		CompositeProcessor = NewTestProcessor<UMassTestCompositeObserverProcessor>(EntityManager);
+		CompositeProcessor->SetObservedTypesForTest({ FTestFragment_Float::StaticStruct(), FTestFragment_Int::StaticStruct() });
+		CompositeProcessor->SetObservedOperationsForTest(EMassObservedOperationFlags::Add);
+		return true;
+	}
+
+	virtual bool InstantTest() override
+	{
+		FMassObserverManager& ObserverManager = EntityManager->GetObserverManager();
+		ObserverManager.AddObserverInstance(CompositeProcessor);
+
+		// First move: observer should fire.
+		const FMassEntityHandle Entity1 = EntityManager->CreateEntity(EmptyArchetype);
+		EntityManager->MoveEntityToAnotherArchetype(Entity1, FloatsIntsArchetype);
+		AITEST_EQUAL(TEXT("Observer should fire once before removal"), CompositeProcessor->ExecuteCallCount, 1);
+
+		// Remove from all type pipelines via the base API.
+		ObserverManager.RemoveObserverInstance(CompositeProcessor);
+
+		// Second move: observer should NOT fire.
+		const FMassEntityHandle Entity2 = EntityManager->CreateEntity(EmptyArchetype);
+		EntityManager->MoveEntityToAnotherArchetype(Entity2, FloatsIntsArchetype);
+		AITEST_EQUAL(TEXT("Observer should not fire after removal"), CompositeProcessor->ExecuteCallCount, 1);
+
+		return true;
+	}
+};
+IMPLEMENT_AI_INSTANT_TEST(FCompositeObserver_BaseApiRemoveObserver, "System.Mass.Observer.CompositeObserver.BaseApiRemoveObserver");
+
+/**
+ * Verifies that GetObservedTypes() on UMassTestCompositeObserverProcessor returns exactly the types
+ * that were configured via SetObservedTypesForTest().
+ */
+struct FCompositeObserver_GetObservedTypes : FEntityTestBase
+{
+	virtual bool InstantTest() override
+	{
+		UMassTestCompositeObserverProcessor* Proc = NewTestProcessor<UMassTestCompositeObserverProcessor>(EntityManager);
+
+		AITEST_EQUAL(TEXT("Freshly created composite observer should have no observed types"), Proc->GetObservedTypes().Num(), 0);
+
+		Proc->SetObservedTypesForTest({ FTestFragment_Float::StaticStruct(), FTestFragment_Int::StaticStruct(), FTestTag_A::StaticStruct() });
+
+		AITEST_EQUAL(TEXT("Composite observer should report 3 observed types"), Proc->GetObservedTypes().Num(), 3);
+		AITEST_TRUE(TEXT("First observed type should be Float"), Proc->GetObservedTypes()[0] == FTestFragment_Float::StaticStruct());
+		AITEST_TRUE(TEXT("Second observed type should be Int"), Proc->GetObservedTypes()[1] == FTestFragment_Int::StaticStruct());
+		AITEST_TRUE(TEXT("Third observed type should be Tag_A"), Proc->GetObservedTypes()[2] == FTestTag_A::StaticStruct());
+
+		return true;
+	}
+};
+IMPLEMENT_AI_INSTANT_TEST(FCompositeObserver_GetObservedTypes, "System.Mass.Observer.CompositeObserver.GetObservedTypes");
+
 } // UE::Mass::Test::Observers
 
 UE_ENABLE_OPTIMIZATION_SHIP
diff --git a/Engine/Source/Developer/MassEntityTestSuite/Public/MassEntityTestTypes.h b/Engine/Source/Developer/MassEntityTestSuite/Public/MassEntityTestTypes.h
--- a/Engine/Source/Developer/MassEntityTestSuite/Public/MassEntityTestTypes.h
+++ b/Engine/Source/Developer/MassEntityTestSuite/Public/MassEntityTestTypes.h
@@ -3,6 +3,7 @@
 #pragma once
 
 #include "MassProcessor.h"
+#include "MassObserverProcessor.h"
 #include "MassEntityQuery.h"
 #include "MassEntityTypes.h"
 #include "MassEntityManager.h"
@@ -625,4 +626,45 @@ struct FMassTestProcessingPhaseManager : public FMassProcessingPhaseManager
 
 } // namespace UE::Mass::Testing
 
+/**
+ * Minimal concrete multi-type observer processor for unit tests.
+ * Counts invocations of Execute() to verify deduplication guarantees.
+ * bAutoRegisterWithObserverRegistry is set to false in the constructor so that the CDO
+ * produced during module init does not attempt to contact UMassObserverRegistry.
+ */
+UCLASS(MinimalAPI)
+class UMassTestCompositeObserverProcessor : public UMassObserverProcessor
+{
+	GENERATED_BODY()
+public:
+	UMassTestCompositeObserverProcessor()
+	{
+		bAutoRegisterWithObserverRegistry = false;
+		QueryBasedPruning = EMassQueryBasedPruning::Never;
+	}
+
+	int32 ExecuteCallCount = 0;
+
+	virtual void Execute(FMassEntityManager& EntityManager, FMassExecutionContext& Context) override
+	{
+		++ExecuteCallCount;
+	}
+
+	/** Sets the observed types list directly for test configuration. */
+	void SetObservedTypesForTest(std::initializer_list<const UScriptStruct*> Types)
+	{
+		ObservedTypes.Reset(static_cast<int32>(Types.size()));
+		for (const UScriptStruct* T : Types)
+		{
+			ObservedTypes.Add(T);
+		}
+	}
+
+	/** Sets the observed operation flags directly for test configuration. */
+	void SetObservedOperationsForTest(const EMassObservedOperationFlags InFlags)
+	{
+		ObservedOperations = InFlags;
+	}
+};
+
 #undef UE_API
diff --git a/Engine/Source/Runtime/Mass/MassEngine/Private/AI/MassEngineNavigationProcessors.cpp b/Engine/Source/Runtime/Mass/MassEngine/Private/AI/MassEngineNavigationProcessors.cpp
--- a/Engine/Source/Runtime/Mass/MassEngine/Private/AI/MassEngineNavigationProcessors.cpp
+++ b/Engine/Source/Runtime/Mass/MassEngine/Private/AI/MassEngineNavigationProcessors.cpp
@@ -296,7 +296,7 @@ UMassNavigationRelevantMassDeinitializer::UMassNavigationRelevantMassDeinitializ
 	: EntityQuery(*this)
 {
 	ExecutionFlags = static_cast<int32>(EProcessorExecutionFlags::AllWorldModes);
-	ObservedType = FMassNavigationRelevantFragment::StaticStruct();
+	ObservedTypes.Add(FMassNavigationRelevantFragment::StaticStruct());
 	ObservedOperations = EMassObservedOperationFlags::Remove;
 	bRequiresGameThreadExecution = true;
 }
diff --git a/Engine/Source/Runtime/Mass/MassEngine/Private/Mesh/MassEngineRenderISMProcessors.cpp b/Engine/Source/Runtime/Mass/MassEngine/Private/Mesh/MassEngineRenderISMProcessors.cpp
--- a/Engine/Source/Runtime/Mass/MassEngine/Private/Mesh/MassEngineRenderISMProcessors.cpp
+++ b/Engine/Source/Runtime/Mass/MassEngine/Private/Mesh/MassEngineRenderISMProcessors.cpp
@@ -391,7 +391,7 @@ UMassRemoveInstanceAndDestroyISMRenderStateProcessor::UMassRemoveInstanceAndDest
 	: EntityQuery(*this)
 {
 	ExecutionFlags = static_cast<int32>(EProcessorExecutionFlags::AllWorldModes);
-	ObservedType = FMassRenderISMInstanceFragment::StaticStruct();
+	ObservedTypes.Add(FMassRenderISMInstanceFragment::StaticStruct());
 	ObservedOperations = EMassObservedOperationFlags::Remove;
 	bRequiresGameThreadExecution = true; // @todo revisit if the execution is required on game thread
 }
diff --git a/Engine/Source/Runtime/Mass/MassEngine/Private/Mesh/MassEngineRenderStateProcessors.cpp b/Engine/Source/Runtime/Mass/MassEngine/Private/Mesh/MassEngineRenderStateProcessors.cpp
--- a/Engine/Source/Runtime/Mass/MassEngine/Private/Mesh/MassEngineRenderStateProcessors.cpp
+++ b/Engine/Source/Runtime/Mass/MassEngine/Private/Mesh/MassEngineRenderStateProcessors.cpp
@@ -170,7 +170,7 @@ UMassDestroyRenderStateProcessor::UMassDestroyRenderStateProcessor()
 	: EntityQuery(*this)
 {
 	ExecutionFlags = static_cast<int32>(EProcessorExecutionFlags::AllWorldModes);
-	ObservedType = FMassRenderStateFragment::StaticStruct();
+	ObservedTypes.Add(FMassRenderStateFragment::StaticStruct());
 	ObservedOperations = EMassObservedOperationFlags::Remove;
 	bRequiresGameThreadExecution = true; // @todo revisit if the execution is required on game thread
 }
diff --git a/Engine/Source/Runtime/MassEntity/Private/MassObserverManager.cpp b/Engine/Source/Runtime/MassEntity/Private/MassObserverManager.cpp
--- a/Engine/Source/Runtime/MassEntity/Private/MassObserverManager.cpp
+++ b/Engine/Source/Runtime/MassEntity/Private/MassObserverManager.cpp
@@ -477,11 +477,33 @@ void FMassObserverManager::HandleElementsImpl(UE::Mass::FProcessingContext& Proc
 	// @todo maybe we could make this type configurable, so that project-specific code can extend it? 
 	FMassObserverExecutionContext& RuntimeObserverContext = ProcessingContext.AuxData.InitializeAs<FMassObserverExecutionContext>(MoveTemp(ObserverContext));
 
+	using UE::Mass::ObserverManager::Tweakables::InlineAllocatorElementsForOverlapTypes;
+
+	// Prevent the same instance from executing more than once per batch.
+	// Multi-type observer processors appear in multiple type pipelines;
+	// without this guard Execute() would fire once per matched type.
+	// Linear scan is intentional: typical element count is 3-15,
+	// well below the crossover point where hashing would win.
+	TArray<UMassProcessor*, TInlineAllocator<InlineAllocatorElementsForOverlapTypes>> SeenProcessors;
+	TArray<UMassProcessor*, TInlineAllocator<InlineAllocatorElementsForOverlapTypes>> NewProcessors;
+
 	for (UE::Mass::ObserverManager::FObserverContextIterator ContextIterator(RuntimeObserverContext); ContextIterator; ++ContextIterator)
 	{
 		FMassRuntimePipeline& Pipeline = (*HandlersContainer).FindChecked(RuntimeObserverContext.GetCurrentType());
 
-		UE::Mass::Executor::RunProcessorsView(Pipeline.GetMutableProcessors(), ProcessingContext, EntityCollections);
+		NewProcessors.Reset();
+		for (UMassProcessor* Proc : Pipeline.GetMutableProcessors())
+		{
+			if (!SeenProcessors.Contains(Proc))
+			{
+				NewProcessors.Add(Proc);
+				SeenProcessors.Add(Proc);
+			}
+		}
+		if (NewProcessors.Num() > 0)
+		{
+			UE::Mass::Executor::RunProcessorsView(NewProcessors, ProcessingContext, EntityCollections);
+		}
 	}
 }
 
@@ -526,7 +548,18 @@ void FMassObserverManager::AddObserverInstance(const UScriptStruct& ElementType,
 
 void FMassObserverManager::AddObserverInstance(TNotNull<UMassObserverProcessor*> ObserverProcessor)
 {
-	AddObserverInstance(ObserverProcessor->GetObservedTypeChecked(), ObserverProcessor->GetObservedOperations(), ObserverProcessor);
+	for (const UScriptStruct* Type : ObserverProcessor->GetObservedTypes())
+	{
+		AddObserverInstance(Type, ObserverProcessor->GetObservedOperations(), ObserverProcessor);
+	}
+}
+
+void FMassObserverManager::RemoveObserverInstance(TNotNull<UMassObserverProcessor*> ObserverProcessor)
+{
+	for (const UScriptStruct* Type : ObserverProcessor->GetObservedTypes())
+	{
+		RemoveObserverInstance(Type, ObserverProcessor->GetObservedOperations(), ObserverProcessor);
+	}
 }
 
 void FMassObserverManager::RemoveObserverInstance(TNotNull<const UScriptStruct*> ElementType, const EMassObservedOperationFlags OperationFlags, TNotNull<UMassProcessor*> ObserverProcessor)
diff --git a/Engine/Source/Runtime/MassEntity/Private/MassObserverProcessor.cpp b/Engine/Source/Runtime/MassEntity/Private/MassObserverProcessor.cpp
--- a/Engine/Source/Runtime/MassEntity/Private/MassObserverProcessor.cpp
+++ b/Engine/Source/Runtime/MassEntity/Private/MassObserverProcessor.cpp
@@ -33,27 +33,53 @@ PRAGMA_DISABLE_DEPRECATION_WARNINGS
 				: EMassObservedOperationFlags::Remove;
 			Operation = EMassObservedOperation::MAX;
 		}
+
+		// Merge legacy single ObservedType into ObservedTypes array for backward compat
+		if (ObservedType != nullptr)
+		{
+			ObservedTypes.AddUnique(ObservedType);
+			ObservedType = nullptr;
+		}
 PRAGMA_ENABLE_DEPRECATION_WARNINGS
 
-		if (ObservedType != nullptr && ObservedOperations != EMassObservedOperationFlags::None)
+		if (ObservedOperations != EMassObservedOperationFlags::None)
 		{
 			Register();
 		}
 		else if (bAutoRegisterWithObserverRegistry)
 		{
-			UE_LOG(LogMass, Error, TEXT("%hs attempting to register %s while it\'s misconfigured, Type: %s, OperationFlags: %#x")
-				, __FUNCTION__, *MyClass->GetName(), *GetNameSafe(ObservedType), ObservedOperations);
+			UE_LOG(LogMass, Error, TEXT("%hs attempting to register %s while it\'s misconfigured, OperationFlags: %#x")
+				, __FUNCTION__, *MyClass->GetName(), ObservedOperations);
 		}
 	}
 }
 
 void UMassObserverProcessor::Register()
 {
-	if (bAutoRegisterWithObserverRegistry)
+	if (!bAutoRegisterWithObserverRegistry)
 	{
-		check(ObservedType);
-		UMassObserverRegistry::GetMutable().RegisterObserver(ObservedType, static_cast<uint8>(ObservedOperations), GetClass());
+		return;
 	}
-}
 
+	// remove nulls
+	const int32 NumRemoved = ObservedTypes.RemoveAll([](const TObjectPtr<const UScriptStruct>& Element)
+		{
+			return Element.Get() == nullptr;
+		});
+	ensureMsgf(NumRemoved == 0, TEXT("Null entries in ObservedTypes for %s"), *GetClass()->GetName());
 
+	if (ObservedTypes.IsEmpty())
+	{
+		UE_LOGF(LogMass, Error
+			, "%hs: %ls has bAutoRegisterWithObserverRegistry=true but no observed types configured"
+			, __FUNCTION__, *GetClass()->GetName());
+		return;
+	}
+
+	UMassObserverRegistry& Registry = UMassObserverRegistry::GetMutable();
+	const uint8 OperationFlags = static_cast<uint8>(ObservedOperations);
+	for (const UScriptStruct* Type : ObservedTypes)
+	{
+		Registry.RegisterObserver(Type, OperationFlags, GetClass());
+	}
+}
diff --git a/Engine/Source/Runtime/MassEntity/Private/MassRelationObservers.cpp b/Engine/Source/Runtime/MassEntity/Private/MassRelationObservers.cpp
--- a/Engine/Source/Runtime/MassEntity/Private/MassRelationObservers.cpp
+++ b/Engine/Source/Runtime/MassEntity/Private/MassRelationObservers.cpp
@@ -24,18 +24,20 @@ void UMassRelationObserver::InitializeInternal(UObject& Owner, const TSharedRef<
 
 void UMassRelationObserver::ConfigureQueries(const TSharedRef<FMassEntityManager>& EntityManager)
 {
-	check(ObservedType);
-	checkf(UE::Mass::IsA<FMassFragment>(ObservedType) || UE::Mass::IsA<FMassTag>(ObservedType)
+	checkf(!ObservedTypes.IsEmpty(), TEXT("ObservedTypes must be populated before ConfigureQueries is called"));
+	const UScriptStruct* PrimaryObservedType = GetSingleObservedTypeChecked();
+	check(PrimaryObservedType);
+	checkf(UE::Mass::IsA<FMassFragment>(PrimaryObservedType) || UE::Mass::IsA<FMassTag>(PrimaryObservedType)
 		, TEXT("Only tags and fragments are valid observed types for RelationObservers. Received %s")
-		, *ObservedType->GetName());
+		, *PrimaryObservedType->GetName());
 
-	if (UE::Mass::IsA<FMassTag>(ObservedType))
+	if (UE::Mass::IsA<FMassTag>(PrimaryObservedType))
 	{
-		EntityQuery.AddTagRequirement(*ObservedType, EMassFragmentPresence::All);
+		EntityQuery.AddTagRequirement(*PrimaryObservedType, EMassFragmentPresence::All);
 	}
-	else // UE::Mass::IsA<FMassFragment>(ObservedType)
+	else // UE::Mass::IsA<FMassFragment>(PrimaryObservedType)
 	{
-		EntityQuery.AddRequirement(ObservedType, ObservedTypeAccess);
+		EntityQuery.AddRequirement(PrimaryObservedType, ObservedTypeAccess);
 	}
 
 	if (bAutoAddRelationFragmentRequirement || bAutoAddRelationTagRequirement)
@@ -46,7 +48,7 @@ void UMassRelationObserver::ConfigureQueries(const TSharedRef<FMassEntityManager
 		if (bAutoAddRelationTagRequirement)
 		{
 			const TNotNull<const UScriptStruct*> RelationTag = RelationData.Traits.GetRelationTagType();
-			if (RelationTag != ObservedType)
+			if (RelationTag != PrimaryObservedType)
 			{
 				EntityQuery.AddTagRequirement(*RelationTag, EMassFragmentPresence::All);
 			}
@@ -55,7 +57,7 @@ void UMassRelationObserver::ConfigureQueries(const TSharedRef<FMassEntityManager
 		if (bAutoAddRelationFragmentRequirement)
 		{
 			const TNotNull<const UScriptStruct*> RelationFragmentType = RelationData.Traits.RelationFragmentType;
-			if (UE::Mass::IsA<FMassFragment>(RelationFragmentType) && RelationFragmentType != ObservedType)
+			if (UE::Mass::IsA<FMassFragment>(RelationFragmentType) && RelationFragmentType != PrimaryObservedType)
 			{
 				EntityQuery.AddRequirement(RelationFragmentType, RelationFragmentAccessType, EMassFragmentPresence::All);
 			}
@@ -63,13 +65,14 @@ void UMassRelationObserver::ConfigureQueries(const TSharedRef<FMassEntityManager
 	}
 
 	DebugDescription = FString::Printf(TEXT("RelationType: %s ObservedType: %s Operation: %s")
-		, *RelationTypeHandle.ToString(), *ObservedType->GetName()
+		, *RelationTypeHandle.ToString(), *PrimaryObservedType->GetName()
 		, *LexToString(ObservedOperations));
 }
 
 bool UMassRelationObserver::ConfigureRelationObserver(UE::Mass::FTypeHandle InRegisteredTypeHandle, const UE::Mass::FRelationTypeTraits& Traits)
 {
-	ObservedType = &const_cast<UScriptStruct&>(*Traits.RelationFragmentType);
+	ObservedTypes.Reset();
+	ObservedTypes.Add(Traits.RelationFragmentType);
 	RelationTypeHandle = InRegisteredTypeHandle;
 	return true;
 }
@@ -97,7 +100,7 @@ void UMassRelationEntityCreation::Execute(FMassEntityManager& EntityManager, FMa
 	};
 
 	TArray<FRelationInstanceRegistration> RelationInstances;
-	EntityQuery.ForEachEntityChunk(Context, [&RelationInstances, ObservedType = ObservedType](FMassExecutionContext& ExecutionContext)
+	EntityQuery.ForEachEntityChunk(Context, [&RelationInstances, ObservedType = GetSingleObservedTypeChecked()](FMassExecutionContext& ExecutionContext)
 	{
 		TConstArrayView<FMassRelationFragment> RelationFragments = ExecutionContext.GetFragmentView<FMassRelationFragment>(ObservedType);
 		for (FMassExecutionContext::FEntityIterator EntityIt = ExecutionContext.CreateEntityIterator(); EntityIt; ++EntityIt)
@@ -142,7 +145,7 @@ void UMassRelationEntityDestruction::Execute(FMassEntityManager& EntityManager,
 {
 	TArray<TPair<FMassEntityHandle, FMassEntityHandle>> EntitiesToClearOut;
 
-	EntityQuery.ForEachEntityChunk(Context, [&EntitiesToClearOut, ObservedType = ObservedType](FMassExecutionContext& ExecutionContext)
+	EntityQuery.ForEachEntityChunk(Context, [&EntitiesToClearOut, ObservedType = GetSingleObservedTypeChecked()](FMassExecutionContext& ExecutionContext)
 	{
 		TConstArrayView<FMassRelationFragment> RelationFragments = ExecutionContext.GetFragmentView<FMassRelationFragment>(ObservedType);
 		for (FMassExecutionContext::FEntityIterator EntityIt = ExecutionContext.CreateEntityIterator(); EntityIt; ++EntityIt)
@@ -218,14 +221,15 @@ bool UMassRelationRoleDestruction::ConfigureRelationObserver(UE::Mass::FTypeHand
 
 	// note there we're deliberately not calling Super::ConfigureRelationObserver, since we're going to observe role-specific elements
 	RelationTypeHandle = InRegisteredTypeHandle;
+	ObservedTypes.Reset();
 	if (Traits.RoleTraits[RoleAsIndex].Element)
 	{
-		ObservedType = Traits.RoleTraits[RoleAsIndex].Element;
+		ObservedTypes.Add(Traits.RoleTraits[RoleAsIndex].Element);
 	}
 	else
 	{
 		TNotNull<const UScriptStruct*> RelationTypeTag = RelationTypeHandle.GetScriptStruct();
-		ObservedType = RelationTypeTag;
+		ObservedTypes.Add(RelationTypeTag);
 	}
 
 	switch (Traits.RoleTraits[RoleAsIndex].DestructionPolicy)
diff --git a/Engine/Source/Runtime/MassEntity/Private/Relations/MassChildOf.cpp b/Engine/Source/Runtime/MassEntity/Private/Relations/MassChildOf.cpp
--- a/Engine/Source/Runtime/MassEntity/Private/Relations/MassChildOf.cpp
+++ b/Engine/Source/Runtime/MassEntity/Private/Relations/MassChildOf.cpp
@@ -71,7 +71,7 @@ void UMassChildOfRelationEntityCreation::Execute(FMassEntityManager& EntityManag
 	};
 
 	TArray<FRelationInstanceRegistration> RelationInstances;
-	EntityQuery.ForEachEntityChunk(Context, [&RelationInstances, ObservedType = ObservedType](FMassExecutionContext& ExecutionContext)
+	EntityQuery.ForEachEntityChunk(Context, [&RelationInstances, ObservedType = GetSingleObservedTypeChecked()](FMassExecutionContext& ExecutionContext)
 	{
 		TConstArrayView<FMassRelationFragment> RelationFragments = ExecutionContext.GetFragmentView<FMassRelationFragment>(ObservedType);
 		for (FMassExecutionContext::FEntityIterator EntityIt = ExecutionContext.CreateEntityIterator(); EntityIt; ++EntityIt)
diff --git a/Engine/Source/Runtime/MassEntity/Public/MassObserverManager.h b/Engine/Source/Runtime/MassEntity/Public/MassObserverManager.h
--- a/Engine/Source/Runtime/MassEntity/Public/MassObserverManager.h
+++ b/Engine/Source/Runtime/MassEntity/Public/MassObserverManager.h
@@ -202,6 +202,7 @@ public:
 	MASSENTITY_API void AddObserverInstance(TNotNull<const UScriptStruct*> ElementType, EMassObservedOperationFlags OperationFlags, TNotNull<UMassProcessor*> ObserverProcessor);
 	MASSENTITY_API void AddObserverInstance(const UScriptStruct& ElementType, EMassObservedOperation Operation, UMassProcessor& ObserverProcessor);
 	MASSENTITY_API void AddObserverInstance(TNotNull<UMassObserverProcessor*> ObserverProcessor);
+	MASSENTITY_API void RemoveObserverInstance(TNotNull<UMassObserverProcessor*> ObserverProcessor);
 	MASSENTITY_API void RemoveObserverInstance(TNotNull<const UScriptStruct*> ElementType, EMassObservedOperationFlags OperationFlags, TNotNull<UMassProcessor*> ObserverProcessor);
 	MASSENTITY_API void RemoveObserverInstance(const UScriptStruct& ElementType, EMassObservedOperation Operation, UMassProcessor& ObserverProcessor);
 
diff --git a/Engine/Source/Runtime/MassEntity/Public/MassObserverProcessor.h b/Engine/Source/Runtime/MassEntity/Public/MassObserverProcessor.h
--- a/Engine/Source/Runtime/MassEntity/Public/MassObserverProcessor.h
+++ b/Engine/Source/Runtime/MassEntity/Public/MassObserverProcessor.h
@@ -40,7 +40,7 @@ struct FMassObserverExecutionContext
 
 	const UScriptStruct* GetCurrentType() const
 	{
-		return TypesInOperation[CurrentTypeIndex];
+		return CurrentTypeIndex != INDEX_NONE ? TypesInOperation[CurrentTypeIndex] : nullptr;
 	}
 
 	bool IsValid() const
@@ -58,12 +58,25 @@ private:
 
 /**
  * Base class for Processors that are used as "observers" of entity operations.
- * An observer declares the type of Mass element it cares about (Fragments and Tags supported at the moment) - via
- * the ObservedType property - and the types of operations it wants to be notified of - via ObservedOperations.
+ * An observer declares the types of Mass elements it cares about (Fragments and Tags supported) - via
+ * the ObservedTypes array - and the types of
+ * operations it wants to be notified of - via ObservedOperations.
  *
  * When an observed operation takes place the processor's regular execution will take place, with
  * ExecutionContext's "auxiliary data" (obtained by calling GetAuxData) being filled with an instance of FMassObserverExecutionContext,
- * that can be used to get information about the type being handled and the kind of operation. 
+ * that can be used to get information about the type being handled and the kind of operation.
+ *
+ * To observe multiple types, populate the ObservedTypes array in the constructor:
+ * @code
+ *   UMyObserver::UMyObserver()
+ *   {
+ *       ObservedTypes = { FMyFragment::StaticStruct(), FOtherFragment::StaticStruct() };
+ *       ObservedOperations = EMassObservedOperationFlags::Add;
+ *   }
+ * @endcode
+ *
+ * When multiple observed types are added simultaneously, Execute() fires exactly once per entity batch
+ * thanks to deduplication in FMassObserverManager::HandleElementsImpl.
  */
 UCLASS(MinimalAPI, abstract)
 class UMassObserverProcessor : public UMassProcessor
@@ -74,24 +87,37 @@ public:
 	UE_API UMassObserverProcessor();
 
 	EMassObservedOperationFlags GetObservedOperations() const;
+	/**
+	 * Observers support observing multiple fragments now, but if you know there's only one you
+	 * can fetch it with this function. The function will assert if the processor is configured to observe more types.
+	 */
+	TNotNull<const UScriptStruct*> GetSingleObservedTypeChecked() const;
+	const UScriptStruct* GetSingleObservedType() const;
+	UE_DEPRECATED(5.8, "UMassObserverProcessor::GetObservedTypeChecked is deprecated now. Use GetSingleObservedTypeChecked")
 	TNotNull<const UScriptStruct*> GetObservedTypeChecked() const;
+	TConstArrayView<const UScriptStruct*> GetObservedTypes() const;
 
 protected:
 	UE_API virtual void PostInitProperties() override;
 
-	/** 
-	 * By default, registers this class as Operation observer of ObservedType. Override to register for multiple 
-	 * operations and/or types 
+	/**
+	 * Called from PostInitProperties on the CDO when ObservedOperations is set.
+	 * By default, registers this class for all entries in ObservedTypes. Override to customize registration.
 	 */
 	UE_API virtual void Register();
 
 protected:
-	UPROPERTY(EditDefaultsOnly, Category = Processor, config)
-	bool bAutoRegisterWithObserverRegistry = true;
+	/** Determines which Fragment or Tag type this given UMassObserverProcessor will be observing (single-type legacy). */
+	UE_DEPRECATED(5.8, "UMassObserverProcessor::ObservedType is deprecated. Use ObservedTypes instead")
+	UPROPERTY()
+	TObjectPtr<const UScriptStruct> ObservedType;
 
-	/** Determines which Fragment or Tag type this given UMassObserverProcessor will be observing */
+	/** Fragment or Tag types this observer watches. Populate this array in the constructor. */
 	UPROPERTY()
-	TObjectPtr<const UScriptStruct> ObservedType = nullptr;
+	TArray<TObjectPtr<const UScriptStruct>> ObservedTypes;
+
+	UPROPERTY(EditDefaultsOnly, Category = Processor, config)
+	bool bAutoRegisterWithObserverRegistry = true;
 
 	UE_DEPRECATED(5.7, "UMassObserverProcessor::Operation is deprecated. Use ObservedOperations instead")
 	EMassObservedOperation Operation = EMassObservedOperation::MAX;
@@ -108,9 +134,27 @@ inline EMassObservedOperationFlags UMassObserverProcessor::GetObservedOperations
 	return ObservedOperations;
 }
 
+inline TNotNull<const UScriptStruct*> UMassObserverProcessor::GetSingleObservedTypeChecked() const
+{
+	checkf(!ObservedTypes.IsEmpty(), TEXT("%hs called but ObservedTypes is empty"), __FUNCTION__);
+	ensureMsgf(ObservedTypes.Num() == 1, TEXT("%hs called there are multiple types being observed. Fetching just one doesn't give you the full picture."), __FUNCTION__);
+	return ObservedTypes[0];
+}
+
+inline const UScriptStruct* UMassObserverProcessor::GetSingleObservedType() const
+{
+	ensureMsgf(ObservedTypes.Num() <= 1, TEXT("%hs called there are multiple types being observed. Fetching just one doesn't give you the full picture."), __FUNCTION__);
+	return ObservedTypes.Num() ? ObservedTypes[0].Get() : static_cast<const UScriptStruct*>(nullptr);
+}
+
 inline TNotNull<const UScriptStruct*> UMassObserverProcessor::GetObservedTypeChecked() const
 {
-	return ObservedType;
+	return GetSingleObservedTypeChecked();
+}
+
+inline TConstArrayView<const UScriptStruct*> UMassObserverProcessor::GetObservedTypes() const
+{
+	return ObservedTypes;
 }
 
 #undef UE_API
diff --git a/Engine/Source/Runtime/MassEntity/Public/MassRelationObservers.h b/Engine/Source/Runtime/MassEntity/Public/MassRelationObservers.h
--- a/Engine/Source/Runtime/MassEntity/Public/MassRelationObservers.h
+++ b/Engine/Source/Runtime/MassEntity/Public/MassRelationObservers.h
@@ -33,7 +33,7 @@ protected:
 
 	FMassEntityQuery EntityQuery;
 	UE::Mass::FTypeHandle RelationTypeHandle;
-	/** Relevant only when ObservedType is a fragment */
+	/** Relevant only when the primary observed type is a fragment */
 	EMassFragmentAccess ObservedTypeAccess = EMassFragmentAccess::ReadWrite;
 	EMassFragmentAccess RelationFragmentAccessType = EMassFragmentAccess::ReadWrite;
 	bool bAutoAddRelationFragmentRequirement = true;
