[Mass] Added a bunch of dead simple tests to excersize Mass's API

Note that this tests are not to learn from!

#rnx
#Synced-CL 51878293

diff --git a/Engine/Source/Developer/MassEntityTestSuite/Private/StressTests/MassStressCommandsTest.cpp b/Engine/Source/Developer/MassEntityTestSuite/Private/StressTests/MassStressCommandsTest.cpp
new file mode 100644
--- /dev/null
+++ b/Engine/Source/Developer/MassEntityTestSuite/Private/StressTests/MassStressCommandsTest.cpp
@@ -0,0 +1,862 @@
+// Copyright Epic Games, Inc. All Rights Reserved.
+
+#include "MassEntityManager.h"
+#include "MassEntityTestTypes.h"
+#include "MassEntityView.h"
+#include "MassExecutionContext.h"
+#include "Algo/RandomShuffle.h"
+#include "MassCommandBuffer.h"
+
+#define LOCTEXT_NAMESPACE "MassTest"
+
+UE_DISABLE_OPTIMIZATION_SHIP
+
+//----------------------------------------------------------------------//
+// Sparse Element Command Tests
+//----------------------------------------------------------------------//
+
+namespace UE::Mass::Test::Commands
+{
+#if WITH_MASSENTITY_DEBUG
+
+struct FCommands_AddSparseFragment : FEntityTestBase
+{
+	virtual bool InstantTest() override
+	{
+		const int32 Count = 10;
+		TArray<FMassEntityHandle> Entities;
+		EntityManager->BatchCreateEntities(IntsArchetype, Count, Entities);
+
+		// Add sparse fragments via command
+		for (int i = 0; i < Count; ++i)
+		{
+			EntityManager->Defer().PushUniqueCommand<FMassCommandAddElement>(
+				FTestFragment_SparseInt::StaticStruct()).Add(Entities[i]);
+		}
+
+		// Verify sparse fragments not yet added
+		for (int i = 0; i < Count; ++i)
+		{
+			AITEST_FALSE(TEXT("Sparse fragment should not exist before flush"),
+				EntityManager->DoesEntityHaveElement(Entities[i], FTestFragment_SparseInt::StaticStruct()));
+		}
+
+		EntityManager->FlushCommands();
+
+		// Verify all entities have sparse fragments after flush
+		for (int i = 0; i < Count; ++i)
+		{
+			AITEST_TRUE(TEXT("Sparse fragment should exist after flush"),
+				EntityManager->DoesEntityHaveElement(Entities[i], FTestFragment_SparseInt::StaticStruct()));
+
+			// Verify entities stayed in same archetype
+			AITEST_EQUAL(TEXT("Entity should remain in original archetype"),
+				EntityManager->GetArchetypeForEntity(Entities[i]), IntsArchetype);
+		}
+
+		return true;
+	}
+};
+IMPLEMENT_AI_INSTANT_TEST(FCommands_AddSparseFragment, "System.Mass.Stress.Commands.SparseElements.AddSparseFragment");
+
+
+struct FCommands_AddSparseTag : FEntityTestBase
+{
+	virtual bool InstantTest() override
+	{
+		const int32 Count = 10;
+		TArray<FMassEntityHandle> Entities;
+		EntityManager->BatchCreateEntities(FloatsArchetype, Count, Entities);
+
+		// Add sparse tags via command
+		for (int i = 0; i < Count; ++i)
+		{
+			EntityManager->Defer().PushUniqueCommand<FMassCommandAddElement>(
+				FTestTag_SparseA::StaticStruct()).Add(Entities[i]);
+		}
+
+		EntityManager->FlushCommands();
+
+		// Verify all entities have sparse tags
+		for (int i = 0; i < Count; ++i)
+		{
+			AITEST_TRUE(TEXT("Sparse tag should exist after flush"),
+				EntityManager->DoesEntityHaveElement(Entities[i], FTestTag_SparseA::StaticStruct()));
+
+			// Verify entities stayed in same archetype
+			AITEST_EQUAL(TEXT("Entity should remain in original archetype"),
+				EntityManager->GetArchetypeForEntity(Entities[i]), FloatsArchetype);
+		}
+
+		return true;
+	}
+};
+IMPLEMENT_AI_INSTANT_TEST(FCommands_AddSparseTag, "System.Mass.Stress.Commands.SparseElements.AddSparseTag");
+
+
+struct FCommands_RemoveSparseFragment : FEntityTestBase
+{
+	virtual bool InstantTest() override
+	{
+		const int32 Count = 10;
+		TArray<FMassEntityHandle> Entities;
+		EntityManager->BatchCreateEntities(FloatsArchetype, Count, Entities);
+
+		// Add sparse fragments directly
+		for (int i = 0; i < Count; ++i)
+		{
+			EntityManager->AddSparseElementToEntity<FTestFragment_SparseInt>(Entities[i]);
+		}
+
+		// Verify all have sparse fragments
+		for (int i = 0; i < Count; ++i)
+		{
+			AITEST_TRUE(TEXT("Sparse fragment should exist"),
+				EntityManager->DoesEntityHaveElement(Entities[i], FTestFragment_SparseInt::StaticStruct()));
+		}
+
+		// Remove via command
+		for (int i = 0; i < Count; ++i)
+		{
+			EntityManager->Defer().PushUniqueCommand<FMassCommandRemoveElement>(
+				FTestFragment_SparseInt::StaticStruct()).Add(Entities[i]);
+		}
+
+		EntityManager->FlushCommands();
+
+		// Verify all sparse fragments removed
+		for (int i = 0; i < Count; ++i)
+		{
+			AITEST_FALSE(TEXT("Sparse fragment should be removed after flush"),
+				EntityManager->DoesEntityHaveElement(Entities[i], FTestFragment_SparseInt::StaticStruct()));
+
+			// Verify entities stayed in same archetype
+			AITEST_EQUAL(TEXT("Entity should remain in original archetype"),
+				EntityManager->GetArchetypeForEntity(Entities[i]), FloatsArchetype);
+		}
+
+		return true;
+	}
+};
+IMPLEMENT_AI_INSTANT_TEST(FCommands_RemoveSparseFragment, "System.Mass.Stress.Commands.SparseElements.RemoveSparseFragment");
+
+
+struct FCommands_RemoveSparseTag : FEntityTestBase
+{
+	virtual bool InstantTest() override
+	{
+		const int32 Count = 10;
+		TArray<FMassEntityHandle> Entities;
+		EntityManager->BatchCreateEntities(IntsArchetype, Count, Entities);
+
+		// Add sparse tags directly
+		for (int i = 0; i < Count; ++i)
+		{
+			EntityManager->AddSparseElementToEntity<FTestTag_SparseA>(Entities[i]);
+		}
+
+		// Remove via command
+		for (int i = 0; i < Count; ++i)
+		{
+			EntityManager->Defer().PushUniqueCommand<FMassCommandRemoveElement>(
+				FTestTag_SparseA::StaticStruct()).Add(Entities[i]);
+		}
+
+		EntityManager->FlushCommands();
+
+		// Verify all sparse tags removed
+		for (int i = 0; i < Count; ++i)
+		{
+			AITEST_FALSE(TEXT("Sparse tag should be removed after flush"),
+				EntityManager->DoesEntityHaveElement(Entities[i], FTestTag_SparseA::StaticStruct()));
+		}
+
+		return true;
+	}
+};
+IMPLEMENT_AI_INSTANT_TEST(FCommands_RemoveSparseTag, "System.Mass.Stress.Commands.SparseElements.RemoveSparseTag");
+
+
+struct FCommands_BatchAddSparseFragments : FEntityTestBase
+{
+	virtual bool InstantTest() override
+	{
+		const int32 EntitiesPerChunk = EntityManager->DebugGetArchetypeEntitiesCountPerChunk(FloatsIntsArchetype);
+		const int32 Count = static_cast<int32>(static_cast<float>(EntitiesPerChunk) * 2.5f); // Span multiple chunks
+
+		TArray<FMassEntityHandle> Entities;
+		EntityManager->BatchCreateEntities(FloatsIntsArchetype, Count, Entities);
+
+		// Batch add sparse fragments via single command with multiple entities
+		EntityManager->Defer().PushUniqueCommand<FMassCommandAddElement>(
+			FTestFragment_SparseInt::StaticStruct()).Add(Entities);
+
+		EntityManager->FlushCommands();
+
+		// Verify all entities have sparse fragments
+		for (int i = 0; i < Count; ++i)
+		{
+			AITEST_TRUE(TEXT("All entities should have sparse fragment after batch add"),
+				EntityManager->DoesEntityHaveElement(Entities[i], FTestFragment_SparseInt::StaticStruct()));
+
+			// Verify entities stayed in same archetype
+			AITEST_EQUAL(TEXT("Entity should remain in original archetype"),
+				EntityManager->GetArchetypeForEntity(Entities[i]), FloatsIntsArchetype);
+		}
+
+		return true;
+	}
+};
+IMPLEMENT_AI_INSTANT_TEST(FCommands_BatchAddSparseFragments, "System.Mass.Stress.Commands.SparseElements.BatchAddSparseFragments");
+
+
+struct FCommands_BatchRemoveSparseFragments : FEntityTestBase
+{
+	virtual bool InstantTest() override
+	{
+		const int32 EntitiesPerChunk = EntityManager->DebugGetArchetypeEntitiesCountPerChunk(IntsArchetype);
+		const int32 Count = static_cast<int32>(static_cast<float>(EntitiesPerChunk) * 2.5f);
+
+		TArray<FMassEntityHandle> Entities;
+		EntityManager->BatchCreateEntities(IntsArchetype, Count, Entities);
+
+		// Add sparse fragments to all entities
+		for (int i = 0; i < Count; ++i)
+		{
+			EntityManager->AddSparseElementToEntity<FTestFragment_SparseFloat>(Entities[i]);
+		}
+
+		// Batch remove via single command
+		EntityManager->Defer().PushUniqueCommand<FMassCommandRemoveElement>(
+			FTestFragment_SparseFloat::StaticStruct()).Add(Entities);
+
+		EntityManager->FlushCommands();
+
+		// Verify all sparse fragments removed
+		for (int i = 0; i < Count; ++i)
+		{
+			AITEST_FALSE(TEXT("All sparse fragments should be removed after batch remove"),
+				EntityManager->DoesEntityHaveElement(Entities[i], FTestFragment_SparseFloat::StaticStruct()));
+		}
+
+		return true;
+	}
+};
+IMPLEMENT_AI_INSTANT_TEST(FCommands_BatchRemoveSparseFragments, "System.Mass.Stress.Commands.SparseElements.BatchRemoveSparseFragments");
+
+
+struct FCommands_MixedSparseAndRegularElements : FEntityTestBase
+{
+	virtual bool InstantTest() override
+	{
+		const int32 Count = 10;
+		TArray<FMassEntityHandle> Entities;
+		EntityManager->BatchCreateEntities(IntsArchetype, Count, Entities);
+
+		// Mix regular fragment addition with sparse element addition
+		for (int i = 0; i < Count; ++i)
+		{
+			// Add regular fragment (causes archetype change)
+			EntityManager->Defer().AddFragment<FTestFragment_Float>(Entities[i]);
+			// Add sparse fragment (no archetype change)
+			EntityManager->Defer().PushUniqueCommand<FMassCommandAddElement>(
+				FTestFragment_SparseInt::StaticStruct()).Add(Entities[i]);
+		}
+
+		EntityManager->FlushCommands();
+
+		// Verify entities moved to new archetype and have sparse elements
+		for (int i = 0; i < Count; ++i)
+		{
+			AITEST_EQUAL(TEXT("Entity should be in FloatsInts archetype"),
+				EntityManager->GetArchetypeForEntity(Entities[i]), FloatsIntsArchetype);
+			AITEST_TRUE(TEXT("Entity should have sparse fragment"),
+				EntityManager->DoesEntityHaveElement(Entities[i], FTestFragment_SparseInt::StaticStruct()));
+
+			// Verify regular fragment accessible
+			AITEST_NOT_NULL(TEXT("Entity should have float fragment"),
+				EntityManager->GetFragmentDataPtr<FTestFragment_Float>(Entities[i]));
+		}
+
+		return true;
+	}
+};
+IMPLEMENT_AI_INSTANT_TEST(FCommands_MixedSparseAndRegularElements, "System.Mass.Stress.Commands.SparseElements.MixedSparseAndRegular");
+
+
+struct FCommands_SparseFragmentDataPreservation : FEntityTestBase
+{
+	virtual bool InstantTest() override
+	{
+		const int32 Count = 10;
+		TArray<FMassEntityHandle> Entities;
+		EntityManager->BatchCreateEntities(FloatsArchetype, Count, Entities);
+
+		// Add sparse fragments and set values
+		for (int i = 0; i < Count; ++i)
+		{
+			FTestFragment_SparseInt& SparseData = EntityManager->AddSparseElementToEntity<FTestFragment_SparseInt>(Entities[i]);
+			SparseData.Value = i * 100;
+		}
+
+		// Trigger archetype change by adding regular fragment
+		for (int i = 0; i < Count; ++i)
+		{
+			EntityManager->Defer().AddFragment<FTestFragment_Int>(Entities[i]);
+		}
+
+		EntityManager->FlushCommands();
+
+		// Verify sparse fragment data preserved across archetype change
+		for (int i = 0; i < Count; ++i)
+		{
+			AITEST_TRUE(TEXT("Sparse fragment should still exist"),
+				EntityManager->DoesEntityHaveElement(Entities[i], FTestFragment_SparseInt::StaticStruct()));
+
+			FStructView SparseData = EntityManager->GetMutableSparseElementDataForEntity(
+				FTestFragment_SparseInt::StaticStruct(), Entities[i]);
+			AITEST_EQUAL(TEXT("Sparse fragment data should be preserved"),
+				SparseData.Get<FTestFragment_SparseInt>().Value, i * 100);
+		}
+
+		return true;
+	}
+};
+IMPLEMENT_AI_INSTANT_TEST(FCommands_SparseFragmentDataPreservation, "System.Mass.Stress.Commands.SparseElements.DataPreservation");
+
+
+struct FCommands_SparseDuplicateAddHandling : FEntityTestBase
+{
+	virtual bool InstantTest() override
+	{
+		const int32 Count = 5;
+		TArray<FMassEntityHandle> Entities;
+		EntityManager->BatchCreateEntities(IntsArchetype, Count, Entities);
+
+		// Add same sparse fragment multiple times (should be idempotent)
+		for (int i = 0; i < Count; ++i)
+		{
+			EntityManager->Defer().PushUniqueCommand<FMassCommandAddElement>(
+				FTestFragment_SparseFloat::StaticStruct()).Add(Entities[i]);
+			EntityManager->Defer().PushUniqueCommand<FMassCommandAddElement>(
+				FTestFragment_SparseFloat::StaticStruct()).Add(Entities[i]);
+		}
+
+		EntityManager->FlushCommands();
+
+		// Verify sparse fragment exists (should handle duplicates gracefully)
+		for (int i = 0; i < Count; ++i)
+		{
+			AITEST_TRUE(TEXT("Sparse fragment should exist after duplicate add"),
+				EntityManager->DoesEntityHaveElement(Entities[i], FTestFragment_SparseFloat::StaticStruct()));
+		}
+
+		return true;
+	}
+};
+IMPLEMENT_AI_INSTANT_TEST(FCommands_SparseDuplicateAddHandling, "System.Mass.Stress.Commands.SparseElements.DuplicateAddHandling");
+
+
+struct FCommands_SparseRemoveNonExistent : FEntityTestBase
+{
+	virtual bool InstantTest() override
+	{
+		const int32 Count = 5;
+		TArray<FMassEntityHandle> Entities;
+		EntityManager->BatchCreateEntities(FloatsArchetype, Count, Entities);
+
+		// Try to remove sparse fragments that don't exist (should be safe)
+		for (int i = 0; i < Count; ++i)
+		{
+			EntityManager->Defer().PushUniqueCommand<FMassCommandRemoveElement>(
+				FTestFragment_SparseInt::StaticStruct()).Add(Entities[i]);
+		}
+
+		EntityManager->FlushCommands();
+
+		// Verify entities still valid and in original archetype
+		for (int i = 0; i < Count; ++i)
+		{
+			AITEST_TRUE(TEXT("Entity should still be valid"), EntityManager->IsEntityValid(Entities[i]));
+			AITEST_EQUAL(TEXT("Entity should remain in original archetype"),
+				EntityManager->GetArchetypeForEntity(Entities[i]), FloatsArchetype);
+		}
+
+		return true;
+	}
+};
+IMPLEMENT_AI_INSTANT_TEST(FCommands_SparseRemoveNonExistent, "System.Mass.Stress.Commands.SparseElements.RemoveNonExistent");
+
+
+struct FCommands_SparseMultipleTypes : FEntityTestBase
+{
+	virtual bool InstantTest() override
+	{
+		const int32 Count = 10;
+		TArray<FMassEntityHandle> Entities;
+		EntityManager->BatchCreateEntities(IntsArchetype, Count, Entities);
+
+		// Add multiple sparse element types to same entities
+		for (int i = 0; i < Count; ++i)
+		{
+			EntityManager->Defer().PushUniqueCommand<FMassCommandAddElement>(
+				FTestFragment_SparseInt::StaticStruct()).Add(Entities[i]);
+			EntityManager->Defer().PushUniqueCommand<FMassCommandAddElement>(
+				FTestFragment_SparseFloat::StaticStruct()).Add(Entities[i]);
+			EntityManager->Defer().PushUniqueCommand<FMassCommandAddElement>(
+				FTestTag_SparseA::StaticStruct()).Add(Entities[i]);
+		}
+
+		EntityManager->FlushCommands();
+
+		// Verify all sparse element types exist
+		for (int i = 0; i < Count; ++i)
+		{
+			AITEST_TRUE(TEXT("SparseInt should exist"),
+				EntityManager->DoesEntityHaveElement(Entities[i], FTestFragment_SparseInt::StaticStruct()));
+			AITEST_TRUE(TEXT("SparseFloat should exist"),
+				EntityManager->DoesEntityHaveElement(Entities[i], FTestFragment_SparseFloat::StaticStruct()));
+			AITEST_TRUE(TEXT("SparseTag should exist"),
+				EntityManager->DoesEntityHaveElement(Entities[i], FTestTag_SparseA::StaticStruct()));
+
+			// All on same archetype
+			AITEST_EQUAL(TEXT("Entity should remain in original archetype"),
+				EntityManager->GetArchetypeForEntity(Entities[i]), IntsArchetype);
+		}
+
+		return true;
+	}
+};
+IMPLEMENT_AI_INSTANT_TEST(FCommands_SparseMultipleTypes, "System.Mass.Stress.Commands.SparseElements.MultipleTypes");
+
+
+struct FCommands_SparseAcrossArchetypes : FEntityTestBase
+{
+	virtual bool InstantTest() override
+	{
+		const int32 Count = 10;
+		TArray<FMassEntityHandle> IntEntities;
+		TArray<FMassEntityHandle> FloatEntities;
+		EntityManager->BatchCreateEntities(IntsArchetype, Count, IntEntities);
+		EntityManager->BatchCreateEntities(FloatsArchetype, Count, FloatEntities);
+
+		TArray<FMassEntityHandle> AllEntities;
+		AllEntities.Append(IntEntities);
+		AllEntities.Append(FloatEntities);
+
+		// Add same sparse element to entities from different archetypes
+		EntityManager->Defer().PushUniqueCommand<FMassCommandAddElement>(
+			FTestFragment_SparseInt::StaticStruct()).Add(AllEntities);
+
+		EntityManager->FlushCommands();
+
+		// Verify all entities have the sparse element regardless of archetype
+		for (const FMassEntityHandle& Entity : IntEntities)
+		{
+			AITEST_TRUE(TEXT("IntArchetype entity should have sparse fragment"),
+				EntityManager->DoesEntityHaveElement(Entity, FTestFragment_SparseInt::StaticStruct()));
+			AITEST_EQUAL(TEXT("Entity should remain in Ints archetype"),
+				EntityManager->GetArchetypeForEntity(Entity), IntsArchetype);
+		}
+
+		for (const FMassEntityHandle& Entity : FloatEntities)
+		{
+			AITEST_TRUE(TEXT("FloatArchetype entity should have sparse fragment"),
+				EntityManager->DoesEntityHaveElement(Entity, FTestFragment_SparseInt::StaticStruct()));
+			AITEST_EQUAL(TEXT("Entity should remain in Floats archetype"),
+				EntityManager->GetArchetypeForEntity(Entity), FloatsArchetype);
+		}
+
+		return true;
+	}
+};
+IMPLEMENT_AI_INSTANT_TEST(FCommands_SparseAcrossArchetypes, "System.Mass.Stress.Commands.SparseElements.AcrossArchetypes");
+
+
+struct FCommands_SparseWithEntityDestruction : FEntityTestBase
+{
+	virtual bool InstantTest() override
+	{
+		const int32 Count = 10;
+		TArray<FMassEntityHandle> Entities;
+		EntityManager->BatchCreateEntities(IntsArchetype, Count, Entities);
+
+		// Add sparse fragments
+		for (int i = 0; i < Count; ++i)
+		{
+			EntityManager->AddSparseElementToEntity<FTestFragment_SparseInt>(Entities[i]);
+			FStructView SparseData = EntityManager->GetMutableSparseElementDataForEntity(
+				FTestFragment_SparseInt::StaticStruct(), Entities[i]);
+			SparseData.Get<FTestFragment_SparseInt>().Value = i;
+		}
+
+		// Destroy half the entities via command
+		TArray<FMassEntityHandle> EntitiesToDestroy;
+		for (int i = 0; i < Count / 2; ++i)
+		{
+			EntitiesToDestroy.Add(Entities[i]);
+			EntityManager->Defer().DestroyEntity(Entities[i]);
+		}
+
+		EntityManager->FlushCommands();
+
+		// Verify destroyed entities are invalid
+		for (int i = 0; i < Count / 2; ++i)
+		{
+			AITEST_FALSE(TEXT("Destroyed entity should be invalid"),
+				EntityManager->IsEntityValid(Entities[i]));
+		}
+
+		// Verify remaining entities still have sparse fragments with correct data
+		for (int i = Count / 2; i < Count; ++i)
+		{
+			AITEST_TRUE(TEXT("Remaining entity should still be valid"),
+				EntityManager->IsEntityValid(Entities[i]));
+			AITEST_TRUE(TEXT("Remaining entity should still have sparse fragment"),
+				EntityManager->DoesEntityHaveElement(Entities[i], FTestFragment_SparseInt::StaticStruct()));
+
+			FStructView SparseData = EntityManager->GetMutableSparseElementDataForEntity(
+				FTestFragment_SparseInt::StaticStruct(), Entities[i]);
+			AITEST_EQUAL(TEXT("Sparse data should still be valid"),
+				SparseData.Get<FTestFragment_SparseInt>().Value, i);
+		}
+
+		return true;
+	}
+};
+IMPLEMENT_AI_INSTANT_TEST(FCommands_SparseWithEntityDestruction, "System.Mass.Stress.Commands.SparseElements.WithEntityDestruction");
+
+
+struct FCommands_SparsePartialBatch : FEntityTestBase
+{
+	virtual bool InstantTest() override
+	{
+		const int32 Count = 20;
+		TArray<FMassEntityHandle> Entities;
+		EntityManager->BatchCreateEntities(FloatsArchetype, Count, Entities);
+
+		// Add sparse fragment to first half only
+		TArray<FMassEntityHandle> FirstHalf;
+		for (int i = 0; i < Count / 2; ++i)
+		{
+			FirstHalf.Add(Entities[i]);
+		}
+
+		EntityManager->Defer().PushUniqueCommand<FMassCommandAddElement>(
+			FTestFragment_SparseFloat::StaticStruct()).Add(FirstHalf);
+
+		EntityManager->FlushCommands();
+
+		// Verify first half has sparse fragments
+		for (int i = 0; i < Count / 2; ++i)
+		{
+			AITEST_TRUE(TEXT("First half should have sparse fragment"),
+				EntityManager->DoesEntityHaveElement(Entities[i], FTestFragment_SparseFloat::StaticStruct()));
+		}
+
+		// Verify second half does not have sparse fragments
+		for (int i = Count / 2; i < Count; ++i)
+		{
+			AITEST_FALSE(TEXT("Second half should not have sparse fragment"),
+				EntityManager->DoesEntityHaveElement(Entities[i], FTestFragment_SparseFloat::StaticStruct()));
+		}
+
+		return true;
+	}
+};
+IMPLEMENT_AI_INSTANT_TEST(FCommands_SparsePartialBatch, "System.Mass.Stress.Commands.SparseElements.PartialBatch");
+
+
+struct FCommands_BatchDestroyEntities_WithSparseElements : FEntityTestBase
+{
+	virtual bool InstantTest() override
+	{
+		const int32 Count = 20;
+		TArray<FMassEntityHandle> Entities;
+		EntityManager->BatchCreateEntities(IntsArchetype, Count, Entities);
+
+		// Add sparse fragments and tags to all entities
+		for (int i = 0; i < Count; ++i)
+		{
+			FTestFragment_SparseInt& SparseData = EntityManager->AddSparseElementToEntity<FTestFragment_SparseInt>(Entities[i]);
+			SparseData.Value = i * 10;
+
+			EntityManager->AddSparseElementToEntity<FTestTag_SparseA>(Entities[i]);
+		}
+
+		// Verify sparse elements exist in storage
+		const UE::Mass::FSparseElementsStorage& Storage = EntityManager->DebugGetSparseElementsStorage();
+		AITEST_EQUAL(TEXT("Storage should have all SparseInt elements"),
+			Storage.GetNumElementsOfType<FTestFragment_SparseInt>(), static_cast<uint32>(Count));
+
+		// Verify each entity has sparse elements
+		for (int i = 0; i < Count; ++i)
+		{
+			AITEST_TRUE(TEXT("Entity should have SparseInt"),
+				EntityManager->DoesEntityHaveElement<FTestFragment_SparseInt>(Entities[i]));
+		}
+
+		// Destroy all entities using BatchDestroyEntities
+		EntityManager->BatchDestroyEntities(Entities);
+
+		// Verify all entities are invalid
+		for (int i = 0; i < Count; ++i)
+		{
+			AITEST_FALSE(TEXT("Entity should be invalid after destruction"),
+				EntityManager->IsEntityValid(Entities[i]));
+		}
+
+		// Verify sparse elements removed from storage
+		AITEST_EQUAL(TEXT("Storage should have no SparseInt elements after destruction"),
+			Storage.GetNumElementsOfType<FTestFragment_SparseInt>(), 0u);
+
+		return true;
+	}
+};
+IMPLEMENT_AI_INSTANT_TEST(FCommands_BatchDestroyEntities_WithSparseElements, "System.Mass.Stress.Commands.BatchDestroy.WithSparseElements");
+
+
+struct FCommands_BatchDestroyEntities_PartialDestruction : FEntityTestBase
+{
+	virtual bool InstantTest() override
+	{
+		const int32 Count = 30;
+		TArray<FMassEntityHandle> Entities;
+		EntityManager->BatchCreateEntities(FloatsArchetype, Count, Entities);
+
+		// Add sparse fragments to all entities with unique values
+		for (int i = 0; i < Count; ++i)
+		{
+			FTestFragment_SparseFloat& SparseFloat = EntityManager->AddSparseElementToEntity<FTestFragment_SparseFloat>(Entities[i]);
+			SparseFloat.Value = static_cast<float>(i);
+		}
+
+		const UE::Mass::FSparseElementsStorage& Storage = EntityManager->DebugGetSparseElementsStorage();
+		AITEST_EQUAL(TEXT("Storage should have all elements initially"),
+			Storage.GetNumElementsOfType<FTestFragment_SparseFloat>(), static_cast<uint32>(Count));
+
+		// Destroy first third of entities
+		TArray<FMassEntityHandle> EntitiesToDestroy;
+		for (int i = 0; i < Count / 3; ++i)
+		{
+			EntitiesToDestroy.Add(Entities[i]);
+		}
+
+		EntityManager->BatchDestroyEntities(EntitiesToDestroy);
+
+		// Verify destroyed entities are invalid
+		for (int i = 0; i < Count / 3; ++i)
+		{
+			AITEST_FALSE(TEXT("Destroyed entity should be invalid"),
+				EntityManager->IsEntityValid(Entities[i]));
+		}
+
+		// Verify remaining entities still valid with correct data
+		const int32 ExpectedRemaining = Count - Count / 3;
+		AITEST_EQUAL(TEXT("Storage should have correct number of remaining elements"),
+			Storage.GetNumElementsOfType<FTestFragment_SparseFloat>(), static_cast<uint32>(ExpectedRemaining));
+
+		for (int i = Count / 3; i < Count; ++i)
+		{
+			AITEST_TRUE(TEXT("Remaining entity should still be valid"),
+				EntityManager->IsEntityValid(Entities[i]));
+			AITEST_TRUE(TEXT("Remaining entity should still have sparse element"),
+				EntityManager->DoesEntityHaveElement<FTestFragment_SparseFloat>(Entities[i]));
+
+			FStructView SparseData = EntityManager->GetMutableSparseElementDataForEntity(
+				FTestFragment_SparseFloat::StaticStruct(), Entities[i]);
+			AITEST_EQUAL(TEXT("Sparse data should be preserved for remaining entities"),
+				SparseData.Get<FTestFragment_SparseFloat>().Value, static_cast<float>(i));
+		}
+
+		return true;
+	}
+};
+IMPLEMENT_AI_INSTANT_TEST(FCommands_BatchDestroyEntities_PartialDestruction, "System.Mass.Stress.Commands.BatchDestroy.PartialDestruction");
+
+
+struct FCommands_BatchDestroyEntityChunks_WithSparseElements : FEntityTestBase
+{
+	virtual bool InstantTest() override
+	{
+		const int32 EntitiesPerChunk = EntityManager->DebugGetArchetypeEntitiesCountPerChunk(FloatsIntsArchetype);
+		const int32 Count = static_cast<int32>(static_cast<float>(EntitiesPerChunk) * 2.5f); // Span multiple chunks
+
+		TArray<FMassEntityHandle> Entities;
+		EntityManager->BatchCreateEntities(FloatsIntsArchetype, Count, Entities);
+
+		// Add multiple sparse element types to all entities
+		for (int i = 0; i < Count; ++i)
+		{
+			FTestFragment_SparseInt& SparseInt = EntityManager->AddSparseElementToEntity<FTestFragment_SparseInt>(Entities[i]);
+			SparseInt.Value = i;
+
+			FTestFragment_SparseFloat& SparseFloat = EntityManager->AddSparseElementToEntity<FTestFragment_SparseFloat>(Entities[i]);
+			SparseFloat.Value = static_cast<float>(i * 2);
+
+			EntityManager->AddSparseElementToEntity<FTestTag_SparseA>(Entities[i]);
+		}
+
+		const UE::Mass::FSparseElementsStorage& Storage = EntityManager->DebugGetSparseElementsStorage();
+		AITEST_EQUAL(TEXT("Storage should have all SparseInt elements"),
+			Storage.GetNumElementsOfType<FTestFragment_SparseInt>(), static_cast<uint32>(Count));
+		AITEST_EQUAL(TEXT("Storage should have all SparseFloat elements"),
+			Storage.GetNumElementsOfType<FTestFragment_SparseFloat>(), static_cast<uint32>(Count));
+
+		// Create entity collection for batch chunk destruction
+		TArray<FMassArchetypeEntityCollection> EntityCollections;
+		UE::Mass::Utils::CreateEntityCollections(*EntityManager, Entities,
+			FMassArchetypeEntityCollection::EDuplicatesHandling::FoldDuplicates, EntityCollections);
+
+		// Destroy all entity chunks
+		EntityManager->BatchDestroyEntityChunks(EntityCollections);
+
+		// Verify all entities are invalid
+		for (int i = 0; i < Count; ++i)
+		{
+			AITEST_FALSE(TEXT("Entity should be invalid after chunk destruction"),
+				EntityManager->IsEntityValid(Entities[i]));
+		}
+
+		// Verify all sparse elements removed from storage
+		AITEST_EQUAL(TEXT("Storage should have no SparseInt elements after chunk destruction"),
+			Storage.GetNumElementsOfType<FTestFragment_SparseInt>(), 0u);
+		AITEST_EQUAL(TEXT("Storage should have no SparseFloat elements after chunk destruction"),
+			Storage.GetNumElementsOfType<FTestFragment_SparseFloat>(), 0u);
+
+		return true;
+	}
+};
+IMPLEMENT_AI_INSTANT_TEST(FCommands_BatchDestroyEntityChunks_WithSparseElements, "System.Mass.Stress.Commands.BatchDestroyChunks.WithSparseElements");
+
+
+struct FCommands_BatchDestroyEntityChunks_PartialChunks : FEntityTestBase
+{
+	virtual bool InstantTest() override
+	{
+		const int32 EntitiesPerChunk = EntityManager->DebugGetArchetypeEntitiesCountPerChunk(IntsArchetype);
+		const int32 Count = static_cast<int32>(static_cast<float>(EntitiesPerChunk) * 3.0f); // 3 full chunks
+
+		TArray<FMassEntityHandle> AllEntities;
+		EntityManager->BatchCreateEntities(IntsArchetype, Count, AllEntities);
+
+		// Add sparse elements to all entities
+		for (int i = 0; i < Count; ++i)
+		{
+			FTestFragment_SparseInt& SparseInt = EntityManager->AddSparseElementToEntity<FTestFragment_SparseInt>(AllEntities[i]);
+			SparseInt.Value = i * 100;
+		}
+
+		const UE::Mass::FSparseElementsStorage& Storage = EntityManager->DebugGetSparseElementsStorage();
+		AITEST_EQUAL(TEXT("Storage should have all elements initially"),
+			Storage.GetNumElementsOfType<FTestFragment_SparseInt>(), static_cast<uint32>(Count));
+
+		// Select only first 2 chunks worth of entities to destroy
+		const int32 EntitiesToDestroy = EntitiesPerChunk * 2;
+		TArray<FMassEntityHandle> EntitiesToDestroyArray;
+		EntitiesToDestroyArray.Reserve(EntitiesToDestroy);
+		for (int i = 0; i < EntitiesToDestroy; ++i)
+		{
+			EntitiesToDestroyArray.Add(AllEntities[i]);
+		}
+
+		// Create entity collection and destroy chunks
+		TArray<FMassArchetypeEntityCollection> EntityCollections;
+		UE::Mass::Utils::CreateEntityCollections(*EntityManager, EntitiesToDestroyArray,
+			FMassArchetypeEntityCollection::EDuplicatesHandling::FoldDuplicates, EntityCollections);
+
+		EntityManager->BatchDestroyEntityChunks(EntityCollections);
+
+		// Verify destroyed entities are invalid
+		for (int i = 0; i < EntitiesToDestroy; ++i)
+		{
+			AITEST_FALSE(TEXT("Destroyed entity should be invalid"),
+				EntityManager->IsEntityValid(AllEntities[i]));
+		}
+
+		// Verify remaining entities (third chunk) still valid with correct data
+		const int32 ExpectedRemaining = Count - EntitiesToDestroy;
+		AITEST_EQUAL(TEXT("Storage should have correct number of remaining elements"),
+			Storage.GetNumElementsOfType<FTestFragment_SparseInt>(), static_cast<uint32>(ExpectedRemaining));
+
+		for (int i = EntitiesToDestroy; i < Count; ++i)
+		{
+			AITEST_TRUE(TEXT("Remaining entity should still be valid"),
+				EntityManager->IsEntityValid(AllEntities[i]));
+			AITEST_TRUE(TEXT("Remaining entity should still have sparse element"),
+				EntityManager->DoesEntityHaveElement<FTestFragment_SparseInt>(AllEntities[i]));
+
+			FStructView SparseData = EntityManager->GetMutableSparseElementDataForEntity(
+				FTestFragment_SparseInt::StaticStruct(), AllEntities[i]);
+			AITEST_EQUAL(TEXT("Sparse data should be preserved for remaining entities"),
+				SparseData.Get<FTestFragment_SparseInt>().Value, i * 100);
+		}
+
+		return true;
+	}
+};
+IMPLEMENT_AI_INSTANT_TEST(FCommands_BatchDestroyEntityChunks_PartialChunks, "System.Mass.Stress.Commands.BatchDestroyChunks.PartialChunks");
+
+
+struct FCommands_BatchDestroyEntities_MixedSparseTypes : FEntityTestBase
+{
+	virtual bool InstantTest() override
+	{
+		constexpr int32 Count = 25;
+		TArray<FMassEntityHandle> Entities;
+		EntityManager->BatchCreateEntities(FloatsIntsArchetype, Count, Entities);
+
+		// Add different combinations of sparse elements to different entities
+		// First third: SparseInt + SparseFloat + SparseTag
+		for (int i = 0; i < Count / 3; ++i)
+		{
+			EntityManager->AddSparseElementToEntity<FTestFragment_SparseInt>(Entities[i]);
+			EntityManager->AddSparseElementToEntity<FTestFragment_SparseFloat>(Entities[i]);
+			EntityManager->AddSparseElementToEntity<FTestTag_SparseA>(Entities[i]);
+		}
+
+		// Middle third: SparseInt + SparseTag only
+		for (int i = Count / 3; i < (Count * 2) / 3; ++i)
+		{
+			EntityManager->AddSparseElementToEntity<FTestFragment_SparseInt>(Entities[i]);
+			EntityManager->AddSparseElementToEntity<FTestTag_SparseA>(Entities[i]);
+		}
+
+		// Last third: SparseFloat only
+		for (int i = (Count * 2) / 3; i < Count; ++i)
+		{
+			EntityManager->AddSparseElementToEntity<FTestFragment_SparseFloat>(Entities[i]);
+		}
+
+		const UE::Mass::FSparseElementsStorage& Storage = EntityManager->DebugGetSparseElementsStorage();
+
+		// Verify initial counts
+		constexpr uint32 InitialSparseIntCount = (Count / 3) + ((Count * 2) / 3 - Count / 3);
+		constexpr uint32 InitialSparseFloatCount = (Count / 3) + (Count - (Count * 2) / 3);
+
+		AITEST_EQUAL(TEXT("Initial SparseInt count should be correct"),
+			Storage.GetNumElementsOfType<FTestFragment_SparseInt>(), InitialSparseIntCount);
+		AITEST_EQUAL(TEXT("Initial SparseFloat count should be correct"),
+			Storage.GetNumElementsOfType<FTestFragment_SparseFloat>(), InitialSparseFloatCount);
+
+		// Destroy all entities
+		EntityManager->BatchDestroyEntities(Entities);
+
+		// Verify all sparse elements removed regardless of type combinations
+		AITEST_EQUAL(TEXT("All SparseInt elements should be removed"),
+			Storage.GetNumElementsOfType<FTestFragment_SparseInt>(), 0u);
+		AITEST_EQUAL(TEXT("All SparseFloat elements should be removed"),
+			Storage.GetNumElementsOfType<FTestFragment_SparseFloat>(), 0u);
+
+		return true;
+	}
+};
+IMPLEMENT_AI_INSTANT_TEST(FCommands_BatchDestroyEntities_MixedSparseTypes, "System.Mass.Stress.Commands.BatchDestroy.MixedSparseTypes");
+
+
+#endif // WITH_MASSENTITY_DEBUG
+} // UE::Mass::Test::Commands
+
+UE_ENABLE_OPTIMIZATION_SHIP
+
+#undef LOCTEXT_NAMESPACE
diff --git a/Engine/Source/Developer/MassEntityTestSuite/Private/StressTests/MassStressEntityBuilderTest.cpp b/Engine/Source/Developer/MassEntityTestSuite/Private/StressTests/MassStressEntityBuilderTest.cpp
new file mode 100644
--- /dev/null
+++ b/Engine/Source/Developer/MassEntityTestSuite/Private/StressTests/MassStressEntityBuilderTest.cpp
@@ -0,0 +1,527 @@
+// Copyright Epic Games, Inc. All Rights Reserved.
+
+#include "AITestsCommon.h"
+#include "MassEntityTestTypes.h"
+#include "MassEntityBuilder.h"
+#include "MassEntityManager.h"
+#include "MassExecutionContext.h"
+#include "MassExecutor.h"
+#include "MassProcessingContext.h"
+
+
+#define LOCTEXT_NAMESPACE "MassTest"
+
+UE_DISABLE_OPTIMIZATION_SHIP
+
+#if 1
+
+/**
+ * Tests to be added:
+ * - test observers triggering as expected, i.e. respecting the construction context
+ * - entity grouping
+ */
+
+namespace UE::Mass::Test::EntityBuilder
+{
+
+//-----------------------------------------------------------------------------
+// Sparse Elements Tests
+//-----------------------------------------------------------------------------
+
+struct FEntityBuilder_AddSparseFrag_Templated : FEntityTestBase
+{
+	virtual bool InstantTest() override
+	{
+		CA_ASSUME(EntityManager);
+
+		// Create entity with sparse fragment using templated Add<>
+		UE::Mass::FEntityBuilder Builder(*EntityManager);
+		FMassEntityHandle Entity = Builder.Add<FTestFragment_SparseInt>(100)
+			.Commit();
+
+		AITEST_TRUE("Entity should be valid", Entity.IsValid());
+
+		// Verify sparse element was added to storage
+		FConstStructView SparseView = EntityManager->GetSparseElementDataForEntity(FTestFragment_SparseInt::StaticStruct(), Entity);
+		AITEST_TRUE("Sparse element should exist in storage", SparseView.IsValid());
+
+		const FTestFragment_SparseInt* SparseData = SparseView.GetPtr<FTestFragment_SparseInt>();
+		AITEST_NOT_NULL("Sparse data should be accessible", SparseData);
+		AITEST_EQUAL("Sparse value should match initialization", SparseData->Value, 100);
+
+		return true;
+	}
+};
+IMPLEMENT_AI_INSTANT_TEST(FEntityBuilder_AddSparseFrag_Templated, "System.Mass.Stress.EntityBuilder.SparseElements.AddSparseFrag_Templated");
+
+struct FEntityBuilder_AddSparseTag_Templated : FEntityTestBase
+{
+	virtual bool InstantTest() override
+	{
+		CA_ASSUME(EntityManager);
+
+		// Create entity with sparse tag using templated Add<>
+		UE::Mass::FEntityBuilder Builder(*EntityManager);
+		FMassEntityHandle Entity = Builder.Add<FTestTag_SparseA>()
+			.Commit();
+
+		AITEST_TRUE("Entity should be valid", Entity.IsValid());
+
+		// Verify sparse tag was added
+		const bool bHasTag = EntityManager->DoesEntityHaveElement(Entity, FTestTag_SparseA::StaticStruct());
+		
+		AITEST_TRUE("Entity should have sparse tag", bHasTag);
+
+		return true;
+	}
+};
+IMPLEMENT_AI_INSTANT_TEST(FEntityBuilder_AddSparseTag_Templated, "System.Mass.Stress.EntityBuilder.SparseElements.AddSparseTag_Templated");
+
+struct FEntityBuilder_AddSparse_Mixed : FEntityTestBase
+{
+	virtual bool InstantTest() override
+	{
+		CA_ASSUME(EntityManager);
+
+		// Create entity with mix of regular and sparse fragments
+		UE::Mass::FEntityBuilder Builder(*EntityManager);
+		FMassEntityHandle Entity = Builder.Add<FTestFragment_Int>(42)
+			.Add<FTestFragment_SparseInt>(100)
+			.Add<FTestFragment_Float>(3.14f)
+			.Add<FTestFragment_SparseFloat>(2.718f)
+			.Add<FTestTag_A>()
+			.Add<FTestTag_SparseB>()
+			.Commit();
+
+		AITEST_TRUE("Entity should be valid", Entity.IsValid());
+
+		// Verify regular fragments
+		const FTestFragment_Int& RegularInt = EntityManager->GetFragmentDataChecked<FTestFragment_Int>(Entity);
+		AITEST_EQUAL("Regular int fragment should match", RegularInt.Value, 42);
+
+		const FTestFragment_Float& RegularFloat = EntityManager->GetFragmentDataChecked<FTestFragment_Float>(Entity);
+		AITEST_EQUAL("Regular float fragment should match", RegularFloat.Value, 3.14f);
+
+		// Verify regular tag
+		AITEST_TRUE("Entity should have regular tag", EntityManager->DoesEntityHaveElement<FTestTag_A>(Entity));
+
+		// Verify sparse fragments
+		FConstStructView SparseIntView = EntityManager->GetSparseElementDataForEntity(FTestFragment_SparseInt::StaticStruct(), Entity);
+		AITEST_TRUE("Sparse int should exist", SparseIntView.IsValid());
+		AITEST_EQUAL("Sparse int value should match", SparseIntView.GetPtr<FTestFragment_SparseInt>()->Value, 100);
+
+		FConstStructView SparseFloatView = EntityManager->GetSparseElementDataForEntity(FTestFragment_SparseFloat::StaticStruct(), Entity);
+		AITEST_TRUE("Sparse float should exist", SparseFloatView.IsValid());
+		AITEST_EQUAL("Sparse float value should match", SparseFloatView.GetPtr<FTestFragment_SparseFloat>()->Value, 2.718f);
+
+		// Verify sparse tag
+		AITEST_TRUE("Entity should have sparse tag", EntityManager->DoesEntityHaveElement(Entity, FTestTag_SparseB::StaticStruct()));
+
+		return true;
+	}
+};
+IMPLEMENT_AI_INSTANT_TEST(FEntityBuilder_AddSparse_Mixed, "System.Mass.Stress.EntityBuilder.SparseElements.AddSparse_Mixed");
+
+struct FEntityBuilder_AddSparse_InstancedStruct : FEntityTestBase
+{
+	virtual bool InstantTest() override
+	{
+		CA_ASSUME(EntityManager);
+
+		// Create sparse fragment instance
+		FInstancedStruct SparseInstance = FInstancedStruct::Make<FTestFragment_SparseInt>(200);
+
+		// Create entity using Add(FInstancedStruct)
+		UE::Mass::FEntityBuilder Builder(*EntityManager);
+		FMassEntityHandle Entity = Builder.Add(SparseInstance)
+			.Commit();
+
+		AITEST_TRUE("Entity should be valid", Entity.IsValid());
+
+		// Verify sparse element
+		FConstStructView SparseView = EntityManager->GetSparseElementDataForEntity(FTestFragment_SparseInt::StaticStruct(), Entity);
+		AITEST_TRUE("Sparse element should exist", SparseView.IsValid());
+		AITEST_EQUAL("Sparse value should match", SparseView.GetPtr<FTestFragment_SparseInt>()->Value, 200);
+
+		return true;
+	}
+};
+IMPLEMENT_AI_INSTANT_TEST(FEntityBuilder_AddSparse_InstancedStruct, "System.Mass.Stress.EntityBuilder.SparseElements.AddSparse_InstancedStruct");
+
+struct FEntityBuilder_AddSparse_UScriptStruct : FEntityTestBase
+{
+	virtual bool InstantTest() override
+	{
+		CA_ASSUME(EntityManager);
+
+		// Add sparse fragment by type pointer (no data initialization)
+		UE::Mass::FEntityBuilder Builder(*EntityManager);
+		FMassEntityHandle Entity = Builder.Add(FTestFragment_SparseInt::StaticStruct())
+			.Commit();
+
+		AITEST_TRUE("Entity should be valid", Entity.IsValid());
+
+		// Verify sparse element exists with default value
+		FConstStructView SparseView = EntityManager->GetSparseElementDataForEntity(FTestFragment_SparseInt::StaticStruct(), Entity);
+		AITEST_TRUE("Sparse element should exist", SparseView.IsValid());
+		AITEST_EQUAL("Sparse value should be default initialized (0)", SparseView.GetPtr<FTestFragment_SparseInt>()->Value, 0);
+
+		return true;
+	}
+};
+IMPLEMENT_AI_INSTANT_TEST(FEntityBuilder_AddSparse_UScriptStruct, "System.Mass.Stress.EntityBuilder.SparseElements.AddSparse_UScriptStruct");
+
+struct FEntityBuilder_AddGetRef_Sparse : FEntityTestBase
+{
+	virtual bool InstantTest() override
+	{
+		CA_ASSUME(EntityManager);
+
+		// Use Add_GetRef to get mutable reference and modify
+		UE::Mass::FEntityBuilder Builder(*EntityManager);
+		FTestFragment_SparseInt& SparseRef = Builder.Add_GetRef<FTestFragment_SparseInt>();
+		SparseRef.Value = 500;
+
+		FMassEntityHandle Entity = Builder.Commit();
+		AITEST_TRUE("Entity should be valid", Entity.IsValid());
+
+		// Verify value was set through reference
+		FConstStructView SparseView = EntityManager->GetSparseElementDataForEntity(FTestFragment_SparseInt::StaticStruct(), Entity);
+		AITEST_TRUE("Sparse element should exist", SparseView.IsValid());
+		AITEST_EQUAL("Sparse value should match reference modification", SparseView.GetPtr<FTestFragment_SparseInt>()->Value, 500);
+
+		return true;
+	}
+};
+IMPLEMENT_AI_INSTANT_TEST(FEntityBuilder_AddGetRef_Sparse, "System.Mass.Stress.EntityBuilder.SparseElements.AddGetRef_Sparse");
+
+struct FEntityBuilder_GetOrCreate_Sparse : FEntityTestBase
+{
+	virtual bool InstantTest() override
+	{
+		CA_ASSUME(EntityManager);
+
+		// Use GetOrCreate twice - should return same instance
+		UE::Mass::FEntityBuilder Builder(*EntityManager);
+		FTestFragment_SparseInt& FirstRef = Builder.GetOrCreate<FTestFragment_SparseInt>(100);
+		FirstRef.Value = 200;  // Modify
+
+		FTestFragment_SparseInt& SecondRef = Builder.GetOrCreate<FTestFragment_SparseInt>(999);  // Should override
+		AITEST_EQUAL("GetOrCreate should override existing value", SecondRef.Value, 999);
+
+		FMassEntityHandle Entity = Builder.Commit();
+		AITEST_TRUE("Entity should be valid", Entity.IsValid());
+
+		// Verify final value
+		FConstStructView SparseView = EntityManager->GetSparseElementDataForEntity(FTestFragment_SparseInt::StaticStruct(), Entity);
+		AITEST_TRUE("Sparse element should exist", SparseView.IsValid());
+		AITEST_EQUAL("Final sparse value should be from second GetOrCreate", SparseView.GetPtr<FTestFragment_SparseInt>()->Value, 999);
+
+		return true;
+	}
+};
+IMPLEMENT_AI_INSTANT_TEST(FEntityBuilder_GetOrCreate_Sparse, "System.Mass.Stress.EntityBuilder.SparseElements.GetOrCreate_Sparse");
+
+struct FEntityBuilder_Find_Sparse : FEntityTestBase
+{
+	virtual bool InstantTest() override
+	{
+		CA_ASSUME(EntityManager);
+
+		// Create builder with sparse fragment
+		UE::Mass::FEntityBuilder Builder(*EntityManager);
+		Builder.Add<FTestFragment_SparseInt>(300);
+
+		// Find should return non-null pointer
+		FTestFragment_SparseInt* Found = Builder.Find<FTestFragment_SparseInt>();
+		AITEST_NOT_NULL("Find should return non-null for added sparse fragment", Found);
+		AITEST_EQUAL("Found value should match", Found->Value, 300);
+
+		// Find for non-added type should return null
+		FTestFragment_SparseFloat* NotFound = Builder.Find<FTestFragment_SparseFloat>();
+		AITEST_NULL("Find should return null for non-added sparse fragment", NotFound);
+
+		return true;
+	}
+};
+IMPLEMENT_AI_INSTANT_TEST(FEntityBuilder_Find_Sparse, "System.Mass.Stress.EntityBuilder.SparseElements.Find_Sparse");
+
+struct FEntityBuilder_MultipleSparseFragments : FEntityTestBase
+{
+	virtual bool InstantTest() override
+	{
+		CA_ASSUME(EntityManager);
+
+		// Create entity with multiple sparse fragments
+		UE::Mass::FEntityBuilder Builder(*EntityManager);
+		FMassEntityHandle Entity = Builder.Add<FTestFragment_SparseInt>(111)
+			.Add<FTestFragment_SparseFloat>(222.0f)
+			.Commit();
+
+		AITEST_TRUE("Entity should be valid", Entity.IsValid());
+
+		// Verify both sparse fragments
+		FConstStructView SparseIntView = EntityManager->GetSparseElementDataForEntity(FTestFragment_SparseInt::StaticStruct(), Entity);
+		AITEST_TRUE("First sparse element should exist", SparseIntView.IsValid());
+		AITEST_EQUAL("First sparse value should match", SparseIntView.GetPtr<FTestFragment_SparseInt>()->Value, 111);
+
+		FConstStructView SparseFloatView = EntityManager->GetSparseElementDataForEntity(FTestFragment_SparseFloat::StaticStruct(), Entity);
+		AITEST_TRUE("Second sparse element should exist", SparseFloatView.IsValid());
+		AITEST_EQUAL("Second sparse value should match", SparseFloatView.GetPtr<FTestFragment_SparseFloat>()->Value, 222.0f);
+
+		return true;
+	}
+};
+IMPLEMENT_AI_INSTANT_TEST(FEntityBuilder_MultipleSparseFragments, "System.Mass.Stress.EntityBuilder.SparseElements.MultipleSparseFragments");
+
+struct FEntityBuilder_MultipleSparseTags : FEntityTestBase
+{
+	virtual bool InstantTest() override
+	{
+		CA_ASSUME(EntityManager);
+
+		// Create entity with multiple sparse tags
+		UE::Mass::FEntityBuilder Builder(*EntityManager);
+		FMassEntityHandle Entity = Builder.Add<FTestTag_SparseA>()
+			.Add<FTestTag_SparseB>()
+			.Commit();
+
+		AITEST_TRUE("Entity should be valid", Entity.IsValid());
+
+		// Verify both sparse tags
+		AITEST_TRUE("Entity should have first sparse tag", EntityManager->DoesEntityHaveElement(Entity, FTestTag_SparseA::StaticStruct()));
+		AITEST_TRUE("Entity should have second sparse tag", EntityManager->DoesEntityHaveElement(Entity, FTestTag_SparseB::StaticStruct()));
+
+		return true;
+	}
+};
+IMPLEMENT_AI_INSTANT_TEST(FEntityBuilder_MultipleSparseTags, "System.Mass.Stress.EntityBuilder.SparseElements.MultipleSparseTags");
+
+struct FEntityBuilder_CopyDataFromEntity_WithSparse : FEntityTestBase
+{
+	virtual bool InstantTest() override
+	{
+		CA_ASSUME(EntityManager);
+
+		// Create source entity with sparse elements
+		UE::Mass::FEntityBuilder SourceBuilder(*EntityManager);
+		FMassEntityHandle SourceEntity = SourceBuilder.Add<FTestFragment_Int>(42)
+			.Add<FTestFragment_SparseInt>(999)
+			.Add<FTestTag_SparseA>()
+			.Commit();
+
+		// Create new entity by copying from source
+		UE::Mass::FEntityBuilder CopyBuilder(*EntityManager);
+		const bool bSuccess = CopyBuilder.CopyDataFromEntity(SourceEntity);
+		AITEST_TRUE("CopyDataFromEntity should succeed", bSuccess);
+
+		FMassEntityHandle CopiedEntity = CopyBuilder.Commit();
+		AITEST_TRUE("Copied entity should be valid", CopiedEntity.IsValid());
+
+		// Verify regular fragment was copied
+		const FTestFragment_Int& CopiedRegular = EntityManager->GetFragmentDataChecked<FTestFragment_Int>(CopiedEntity);
+		AITEST_EQUAL("Regular fragment should be copied", CopiedRegular.Value, 42);
+
+		// Verify sparse fragment was copied
+		FConstStructView CopiedSparseView = EntityManager->GetSparseElementDataForEntity(FTestFragment_SparseInt::StaticStruct(), CopiedEntity);
+		AITEST_TRUE("Sparse fragment should be copied", CopiedSparseView.IsValid());
+		AITEST_EQUAL("Sparse fragment value should match", CopiedSparseView.GetPtr<FTestFragment_SparseInt>()->Value, 999);
+
+		// Verify sparse tag was copied
+		AITEST_TRUE("Sparse tag should be copied", EntityManager->DoesEntityHaveElement(CopiedEntity, FTestTag_SparseA::StaticStruct()));
+
+		return true;
+	}
+};
+IMPLEMENT_AI_INSTANT_TEST(FEntityBuilder_CopyDataFromEntity_WithSparse, "System.Mass.Stress.EntityBuilder.SparseElements.CopyDataFromEntity_WithSparse");
+
+struct FEntityBuilder_AppendDataFromEntity_WithSparse : FEntityTestBase
+{
+	virtual bool InstantTest() override
+	{
+		CA_ASSUME(EntityManager);
+
+		// Create source entity with sparse elements
+		UE::Mass::FEntityBuilder SourceBuilder(*EntityManager);
+		FMassEntityHandle SourceEntity = SourceBuilder.Add<FTestFragment_SparseInt>(777)
+			.Add<FTestTag_SparseA>()
+			.Commit();
+
+		// Create builder with some existing data
+		UE::Mass::FEntityBuilder AppendBuilder(*EntityManager);
+		AppendBuilder.Add<FTestFragment_Float>(1.5f)
+			.Add<FTestTag_SparseB>();
+
+		// Append data from source
+		const bool bSuccess = AppendBuilder.AppendDataFromEntity(SourceEntity);
+		AITEST_TRUE("AppendDataFromEntity should succeed", bSuccess);
+
+		FMassEntityHandle ResultEntity = AppendBuilder.Commit();
+		AITEST_TRUE("Result entity should be valid", ResultEntity.IsValid());
+
+		// Verify original data is still there
+		const FTestFragment_Float& OriginalFloat = EntityManager->GetFragmentDataChecked<FTestFragment_Float>(ResultEntity);
+		AITEST_EQUAL("Original fragment should remain", OriginalFloat.Value, 1.5f);
+		AITEST_TRUE("Original sparse tag should remain", EntityManager->DoesEntityHaveElement(ResultEntity, FTestTag_SparseB::StaticStruct()));
+
+		// Verify appended sparse data
+		FConstStructView AppendedSparseView = EntityManager->GetSparseElementDataForEntity(FTestFragment_SparseInt::StaticStruct(), ResultEntity);
+		AITEST_TRUE("Appended sparse fragment should exist", AppendedSparseView.IsValid());
+		AITEST_EQUAL("Appended sparse value should match", AppendedSparseView.GetPtr<FTestFragment_SparseInt>()->Value, 777);
+
+		AITEST_TRUE("Appended sparse tag should exist", EntityManager->DoesEntityHaveElement(ResultEntity, FTestTag_SparseA::StaticStruct()));
+
+		return true;
+	}
+};
+IMPLEMENT_AI_INSTANT_TEST(FEntityBuilder_AppendDataFromEntity_WithSparse, "System.Mass.Stress.EntityBuilder.SparseElements.AppendDataFromEntity_WithSparse");
+
+struct FEntityBuilder_Make_WithSparseFragments : FEntityTestBase
+{
+	virtual bool InstantTest() override
+	{
+		CA_ASSUME(EntityManager);
+
+		// Prepare composition and initial values
+		FMassArchetypeCompositionDescriptor Composition;
+		Composition.Add<FTestFragment_Int>();
+		Composition.Add<FTestFragment_SparseInt>();
+		Composition.Add<FTestTag_SparseA>();
+
+		TArray<FInstancedStruct> InitialFragments;
+		InitialFragments.Add(FInstancedStruct::Make<FTestFragment_Int>(50));
+		InitialFragments.Add(FInstancedStruct::Make<FTestFragment_SparseInt>(888));
+
+		// Create builder using Make with sparse fragments
+		UE::Mass::FEntityBuilder Builder = UE::Mass::FEntityBuilder::Make(
+			EntityManager->AsShared(),
+			Composition,
+			InitialFragments
+		);
+
+		FMassEntityHandle Entity = Builder.Commit();
+		AITEST_TRUE("Entity should be valid", Entity.IsValid());
+
+		// Verify regular fragment
+		const FTestFragment_Int& RegularInt = EntityManager->GetFragmentDataChecked<FTestFragment_Int>(Entity);
+		AITEST_EQUAL("Regular fragment should match", RegularInt.Value, 50);
+
+		// Verify sparse fragment
+		FConstStructView SparseView = EntityManager->GetSparseElementDataForEntity(FTestFragment_SparseInt::StaticStruct(), Entity);
+		AITEST_TRUE("Sparse fragment should exist", SparseView.IsValid());
+		AITEST_EQUAL("Sparse fragment value should match", SparseView.GetPtr<FTestFragment_SparseInt>()->Value, 888);
+
+		// Verify sparse tag
+		AITEST_TRUE("Sparse tag should exist", EntityManager->DoesEntityHaveElement(Entity, FTestTag_SparseA::StaticStruct()));
+
+		return true;
+	}
+};
+IMPLEMENT_AI_INSTANT_TEST(FEntityBuilder_Make_WithSparseFragments, "System.Mass.Stress.EntityBuilder.SparseElements.Make_WithSparseFragments");
+
+struct FEntityBuilder_CommitAndReprepare_WithSparse : FEntityTestBase
+{
+	virtual bool InstantTest() override
+	{
+		CA_ASSUME(EntityManager);
+
+		// Create builder with sparse elements
+		UE::Mass::FEntityBuilder Builder(*EntityManager);
+		Builder.Add<FTestFragment_SparseInt>(111);
+
+		// Commit first entity
+		FMassEntityHandle Entity1 = Builder.CommitAndReprepare();
+		AITEST_TRUE("First entity should be valid", Entity1.IsValid());
+
+		// Modify and commit second entity
+		FTestFragment_SparseInt& Ref = Builder.GetOrCreate<FTestFragment_SparseInt>(222);
+		AITEST_EQUAL("Second entity data should be modified", Ref.Value, 222);
+
+		FMassEntityHandle Entity2 = Builder.Commit();
+		AITEST_TRUE("Second entity should be valid", Entity2.IsValid());
+		AITEST_NOT_EQUAL("Entities should be different", Entity1.Index, Entity2.Index);
+
+		// Verify both entities have correct values
+		FConstStructView Sparse1 = EntityManager->GetSparseElementDataForEntity(FTestFragment_SparseInt::StaticStruct(), Entity1);
+		AITEST_TRUE("The Sparse1 fragment view is expected to be valid", Sparse1.IsValid());
+		AITEST_EQUAL("First entity should have original value", Sparse1.GetPtr<FTestFragment_SparseInt>()->Value, 111);
+
+		FConstStructView Sparse2 = EntityManager->GetSparseElementDataForEntity(FTestFragment_SparseInt::StaticStruct(), Entity2);
+		AITEST_TRUE("The Sparse2 fragment view is expected to be valid", Sparse2.IsValid());
+		AITEST_EQUAL("Second entity should have modified value", Sparse2.GetPtr<FTestFragment_SparseInt>()->Value, 222);
+
+		return true;
+	}
+};
+IMPLEMENT_AI_INSTANT_TEST(FEntityBuilder_CommitAndReprepare_WithSparse, "System.Mass.Stress.EntityBuilder.SparseElements.CommitAndReprepare_WithSparse");
+
+struct FEntityBuilder_DeferredCommit_WithSparse : FEntityTestBase
+{
+	virtual bool InstantTest() override
+	{
+		CA_ASSUME(EntityManager);
+
+		// Create builder with deferred commit enabled
+		UE::Mass::FEntityBuilder Builder(*EntityManager);
+		Builder.SetForceDeferredCommit(true);
+		Builder.Add<FTestFragment_SparseInt>(333)
+			.Add<FTestTag_SparseA>();
+
+		FMassEntityHandle Entity = Builder.Commit();
+		AITEST_TRUE("Entity should be valid after deferred commit", Entity.IsValid());
+
+		// Flush commands
+		EntityManager->FlushCommands();
+
+		// Verify sparse elements were added through deferred path
+		FConstStructView SparseView = EntityManager->GetSparseElementDataForEntity(FTestFragment_SparseInt::StaticStruct(), Entity);
+		AITEST_TRUE("Sparse fragment should exist after flush", SparseView.IsValid());
+		AITEST_EQUAL("Sparse value should match", SparseView.GetPtr<FTestFragment_SparseInt>()->Value, 333);
+
+		AITEST_TRUE("Sparse tag should exist after flush", EntityManager->DoesEntityHaveElement(Entity, FTestTag_SparseA::StaticStruct()));
+
+		return true;
+	}
+};
+IMPLEMENT_AI_INSTANT_TEST(FEntityBuilder_DeferredCommit_WithSparse, "System.Mass.Stress.EntityBuilder.SparseElements.DeferredCommit_WithSparse");
+
+struct FEntityBuilder_SparseOnly_NoRegularFragments : FEntityTestBase
+{
+	virtual bool InstantTest() override
+	{
+		CA_ASSUME(EntityManager);
+
+		// Create entity with ONLY sparse elements (no regular fragments)
+		UE::Mass::FEntityBuilder Builder(*EntityManager);
+		FMassEntityHandle Entity = Builder.Add<FTestFragment_SparseInt>(444)
+			.Add<FTestFragment_SparseFloat>(555.0f)
+			.Add<FTestTag_SparseA>()
+			.Add<FTestTag_SparseB>()
+			.Commit();
+
+		AITEST_TRUE("Entity should be valid", Entity.IsValid());
+
+		// Verify all sparse elements exist
+		FConstStructView SparseIntView = EntityManager->GetSparseElementDataForEntity(FTestFragment_SparseInt::StaticStruct(), Entity);
+		AITEST_TRUE("Sparse int should exist", SparseIntView.IsValid());
+		AITEST_EQUAL("Sparse int value should match", SparseIntView.GetPtr<FTestFragment_SparseInt>()->Value, 444);
+
+		FConstStructView SparseFloatView = EntityManager->GetSparseElementDataForEntity(FTestFragment_SparseFloat::StaticStruct(), Entity);
+		AITEST_TRUE("Sparse float should exist", SparseFloatView.IsValid());
+		AITEST_EQUAL("Sparse float value should match", SparseFloatView.GetPtr<FTestFragment_SparseFloat>()->Value, 555.0f);
+
+		AITEST_TRUE("First sparse tag should exist", EntityManager->DoesEntityHaveElement(Entity, FTestTag_SparseA::StaticStruct()));
+		AITEST_TRUE("Second sparse tag should exist", EntityManager->DoesEntityHaveElement(Entity, FTestTag_SparseB::StaticStruct()));
+
+		return true;
+	}
+};
+IMPLEMENT_AI_INSTANT_TEST(FEntityBuilder_SparseOnly_NoRegularFragments, "System.Mass.Stress.EntityBuilder.SparseElements.SparseOnly_NoRegularFragments");
+
+} // UE::Mass::Test::EntityBuilder
+
+#endif
+
+UE_ENABLE_OPTIMIZATION_SHIP
+
+#undef LOCTEXT_NAMESPACE
diff --git a/Engine/Source/Developer/MassEntityTestSuite/Private/StressTests/MassStressEntityViewTest.cpp b/Engine/Source/Developer/MassEntityTestSuite/Private/StressTests/MassStressEntityViewTest.cpp
new file mode 100644
--- /dev/null
+++ b/Engine/Source/Developer/MassEntityTestSuite/Private/StressTests/MassStressEntityViewTest.cpp
@@ -0,0 +1,416 @@
+// Copyright Epic Games, Inc. All Rights Reserved.
+
+#include "AITestsCommon.h"
+#include "MassEntityManager.h"
+#include "MassEntityTestTypes.h"
+#include "MassEntityView.h"
+
+#define LOCTEXT_NAMESPACE "MassTest"
+
+UE_DISABLE_OPTIMIZATION_SHIP
+
+namespace UE::Mass::Test::EntityView
+{
+#if WITH_MASSENTITY_DEBUG
+
+//----------------------------------------------------------------------//
+// FMassEntityView Sparse Element Tests
+//----------------------------------------------------------------------//
+
+struct FEntityView_HasSparseFragment_SingleEntity : FEntityTestBase
+{
+	virtual bool InstantTest() override
+	{
+		// Test bug: HasFragment should check per-entity sparse element, not just archetype
+		const FMassEntityHandle Entity1 = EntityManager->CreateEntity(IntsArchetype);
+		const FMassEntityHandle Entity2 = EntityManager->CreateEntity(IntsArchetype);
+
+		// Add sparse fragment to only Entity1
+		EntityManager->AddSparseElementToEntity<FTestFragment_SparseInt>(Entity1);
+
+		FMassEntityView View1(*EntityManager, Entity1);
+		FMassEntityView View2(*EntityManager, Entity2);
+
+		// Both are in same archetype, but only Entity1 has the sparse fragment
+		AITEST_EQUAL(TEXT("Both entities in same archetype"),
+			EntityManager->GetArchetypeForEntity(Entity1),
+			EntityManager->GetArchetypeForEntity(Entity2));
+
+		AITEST_TRUE(TEXT("Entity1 should have sparse fragment"),
+			View1.HasElement<FTestFragment_SparseInt>());
+		AITEST_FALSE(TEXT("Entity2 should NOT have sparse fragment"),
+			View2.HasElement<FTestFragment_SparseInt>());
+
+		return true;
+	}
+};
+IMPLEMENT_AI_INSTANT_TEST(FEntityView_HasSparseFragment_SingleEntity, "System.Mass.Stress.EntityView.Sparse.HasFragment.PerEntity");
+
+
+struct FEntityView_HasSparseTag_SingleEntity : FEntityTestBase
+{
+	virtual bool InstantTest() override
+	{
+		const FMassEntityHandle Entity1 = EntityManager->CreateEntity(FloatsArchetype);
+		const FMassEntityHandle Entity2 = EntityManager->CreateEntity(FloatsArchetype);
+
+		// Add sparse tag to only Entity1
+		EntityManager->AddSparseElementToEntity<FTestTag_SparseA>(Entity1);
+
+		FMassEntityView View1(*EntityManager, Entity1);
+		FMassEntityView View2(*EntityManager, Entity2);
+
+		AITEST_TRUE(TEXT("Entity1 should have sparse tag"),
+			View1.HasElement<FTestTag_SparseA>());
+		AITEST_FALSE(TEXT("Entity2 should NOT have sparse tag"),
+			View2.HasElement<FTestTag_SparseA>());
+
+		return true;
+	}
+};
+IMPLEMENT_AI_INSTANT_TEST(FEntityView_HasSparseTag_SingleEntity, "System.Mass.Stress.EntityView.Sparse.HasTag.PerEntity");
+
+
+struct FEntityView_HasElement_BugTest : FEntityTestBase
+{
+	virtual bool InstantTest() override
+	{
+		// BUG: HasElement doesn't check per-entity sparse elements correctly
+		const FMassEntityHandle Entity1 = EntityManager->CreateEntity(IntsArchetype);
+		const FMassEntityHandle Entity2 = EntityManager->CreateEntity(IntsArchetype);
+
+		// Add sparse fragment to only Entity1
+		EntityManager->AddSparseElementToEntity<FTestFragment_SparseInt>(Entity1);
+
+		FMassEntityView View1(*EntityManager, Entity1);
+		FMassEntityView View2(*EntityManager, Entity2);
+
+		// HasElement with bIncludeSparse=true should check per-entity
+		AITEST_TRUE(TEXT("Entity1.HasElement should return true for sparse fragment"),
+			View1.HasElement<FTestFragment_SparseInt>(UE::Mass::EIncludeSparseElements::Yes));
+		AITEST_FALSE(TEXT("Entity2.HasElement should return false for sparse fragment"),
+			View2.HasElement<FTestFragment_SparseInt>(UE::Mass::EIncludeSparseElements::Yes));
+
+		// HasElement with bIncludeSparse=false should only check archetype
+		AITEST_FALSE(TEXT("Entity1.HasElement(false) should return false"),
+			View1.HasElement<FTestFragment_SparseInt>(UE::Mass::EIncludeSparseElements::No));
+		AITEST_FALSE(TEXT("Entity2.HasElement(false) should return false"),
+			View2.HasElement<FTestFragment_SparseInt>(UE::Mass::EIncludeSparseElements::No));
+
+		return true;
+	}
+};
+IMPLEMENT_AI_INSTANT_TEST(FEntityView_HasElement_BugTest, "System.Mass.Stress.EntityView.Sparse.HasElement.BugTest");
+
+
+struct FEntityView_HasElement_SparseTag : FEntityTestBase
+{
+	virtual bool InstantTest() override
+	{
+		const FMassEntityHandle Entity1 = EntityManager->CreateEntity(FloatsArchetype);
+		const FMassEntityHandle Entity2 = EntityManager->CreateEntity(FloatsArchetype);
+
+		// Add sparse tag to only Entity1
+		EntityManager->AddSparseElementToEntity<FTestTag_SparseA>(Entity1);
+
+		FMassEntityView View1(*EntityManager, Entity1);
+		FMassEntityView View2(*EntityManager, Entity2);
+
+		// HasElement should work with sparse tags too
+		AITEST_TRUE(TEXT("Entity1 HasElement should detect sparse tag"),
+			View1.HasElement<FTestTag_SparseA>(UE::Mass::EIncludeSparseElements::Yes));
+		AITEST_FALSE(TEXT("Entity2 HasElement should not detect sparse tag"),
+			View2.HasElement<FTestTag_SparseA>(UE::Mass::EIncludeSparseElements::Yes));
+
+		return true;
+	}
+};
+IMPLEMENT_AI_INSTANT_TEST(FEntityView_HasElement_SparseTag, "System.Mass.Stress.EntityView.Sparse.HasElement.Tag");
+
+
+struct FEntityView_GetSparseFragmentDataPtr : FEntityTestBase
+{
+	virtual bool InstantTest() override
+	{
+		const FMassEntityHandle Entity = EntityManager->CreateEntity(IntsArchetype);
+
+		// Add sparse fragment and set value
+		FTestFragment_SparseInt& SparseData = EntityManager->AddSparseElementToEntity<FTestFragment_SparseInt>(Entity);
+		SparseData.Value = 42;
+
+		FMassEntityView EntityView(*EntityManager, Entity);
+
+		// Test const version
+		const FTestFragment_SparseInt* ConstFragment = EntityView.GetSparseFragmentDataPtr<FTestFragment_SparseInt>();
+		AITEST_NOT_NULL(TEXT("GetSparseFragmentDataPtr (const) should return valid pointer"), ConstFragment);
+		AITEST_EQUAL(TEXT("Value should be accessible"), ConstFragment->Value, 42);
+
+		// Test mutable version
+		FTestFragment_SparseInt* MutableFragment = EntityView.GetSparseFragmentDataPtr<FTestFragment_SparseInt>();
+		AITEST_NOT_NULL(TEXT("GetSparseFragmentDataPtr (mutable) should return valid pointer"), MutableFragment);
+
+		// Modify value
+		MutableFragment->Value = 100;
+
+		// Verify modification persisted
+		const FTestFragment_SparseInt* VerifyFragment = EntityView.GetSparseFragmentDataPtr<FTestFragment_SparseInt>();
+		AITEST_EQUAL(TEXT("Modified value should persist"), VerifyFragment->Value, 100);
+
+		return true;
+	}
+};
+IMPLEMENT_AI_INSTANT_TEST(FEntityView_GetSparseFragmentDataPtr, "System.Mass.Stress.EntityView.Sparse.GetSparseFragmentDataPtr");
+
+
+struct FEntityView_GetSparseFragmentDataPtr_NonExistent : FEntityTestBase
+{
+	virtual bool InstantTest() override
+	{
+		const FMassEntityHandle Entity = EntityManager->CreateEntity(IntsArchetype);
+
+		FMassEntityView EntityView(*EntityManager, Entity);
+
+		// Get pointer to sparse fragment that doesn't exist
+		const FTestFragment_SparseFloat* Fragment = EntityView.GetSparseFragmentDataPtr<FTestFragment_SparseFloat>();
+		AITEST_NULL(TEXT("GetSparseFragmentDataPtr should return null for non-existent fragment"), Fragment);
+
+		return true;
+	}
+};
+IMPLEMENT_AI_INSTANT_TEST(FEntityView_GetSparseFragmentDataPtr_NonExistent, "System.Mass.Stress.EntityView.Sparse.GetSparseFragmentDataPtr.NonExistent");
+
+
+struct FEntityView_SparseFragments_MultipleTypes : FEntityTestBase
+{
+	virtual bool InstantTest() override
+	{
+		const FMassEntityHandle Entity = EntityManager->CreateEntity(FloatsArchetype);
+
+		// Add multiple sparse types
+		EntityManager->AddSparseElementToEntity<FTestFragment_SparseInt>(Entity);
+		EntityManager->AddSparseElementToEntity<FTestFragment_SparseFloat>(Entity);
+		EntityManager->AddSparseElementToEntity<FTestTag_SparseA>(Entity);
+
+		FMassEntityView EntityView(*EntityManager, Entity);
+
+		// Verify HasElement works for all
+		AITEST_TRUE(TEXT("HasElement detects sparse int"),
+			EntityView.HasElement<FTestFragment_SparseInt>(UE::Mass::EIncludeSparseElements::Yes));
+		AITEST_TRUE(TEXT("HasElement detects sparse float"),
+			EntityView.HasElement<FTestFragment_SparseFloat>(UE::Mass::EIncludeSparseElements::Yes));
+		AITEST_TRUE(TEXT("HasElement detects sparse tag"),
+			EntityView.HasElement<FTestTag_SparseA>(UE::Mass::EIncludeSparseElements::Yes));
+
+		// Verify GetSparseFragmentDataPtr works for all fragment types
+		AITEST_NOT_NULL(TEXT("Can get sparse int pointer"),
+			EntityView.GetSparseFragmentDataPtr<FTestFragment_SparseInt>());
+		AITEST_NOT_NULL(TEXT("Can get sparse float pointer"),
+			EntityView.GetSparseFragmentDataPtr<FTestFragment_SparseFloat>());
+
+		return true;
+	}
+};
+IMPLEMENT_AI_INSTANT_TEST(FEntityView_SparseFragments_MultipleTypes, "System.Mass.Stress.EntityView.Sparse.MultipleTypes");
+
+
+struct FEntityView_SparseFragment_AfterArchetypeChange : FEntityTestBase
+{
+	virtual bool InstantTest() override
+	{
+		const FMassEntityHandle Entity = EntityManager->CreateEntity(IntsArchetype);
+
+		// Add sparse fragment
+		FTestFragment_SparseInt& SparseData = EntityManager->AddSparseElementToEntity<FTestFragment_SparseInt>(Entity);
+		SparseData.Value = 777;
+
+		// Get initial view
+		FMassEntityView View1(*EntityManager, Entity);
+		AITEST_TRUE(TEXT("View1 detects sparse fragment"),
+			View1.HasElement<FTestFragment_SparseInt>());
+
+		// Change archetype by adding regular fragment
+		EntityManager->AddFragmentToEntity(Entity, FTestFragment_Float::StaticStruct());
+
+		// Old view should be invalidated
+		AITEST_FALSE(TEXT("Old view invalidated after archetype change"), View1.IsValid());
+
+		// Create new view
+		FMassEntityView View2(*EntityManager, Entity);
+		AITEST_TRUE(TEXT("New view is valid"), View2.IsValid());
+
+		// Verify sparse fragment still accessible
+		AITEST_TRUE(TEXT("Sparse fragment persisted through archetype change"),
+			View2.HasElement<FTestFragment_SparseInt>());
+
+		const FTestFragment_SparseInt* Fragment = View2.GetSparseFragmentDataPtr<FTestFragment_SparseInt>();
+		AITEST_NOT_NULL(TEXT("Can still get sparse fragment pointer"), Fragment);
+		AITEST_EQUAL(TEXT("Sparse fragment data preserved"), Fragment->Value, 777);
+
+		return true;
+	}
+};
+IMPLEMENT_AI_INSTANT_TEST(FEntityView_SparseFragment_AfterArchetypeChange, "System.Mass.Stress.EntityView.Sparse.AfterArchetypeChange");
+
+
+struct FEntityView_SparseFragment_RemoveAndCheck : FEntityTestBase
+{
+	virtual bool InstantTest() override
+	{
+		const FMassEntityHandle Entity = EntityManager->CreateEntity(IntsArchetype);
+
+		// Add sparse fragment
+		EntityManager->AddSparseElementToEntity<FTestFragment_SparseInt>(Entity);
+
+		FMassEntityView EntityView(*EntityManager, Entity);
+		AITEST_TRUE(TEXT("Fragment exists before removal"),
+			EntityView.HasElement<FTestFragment_SparseInt>());
+
+		// Remove sparse fragment
+		EntityManager->RemoveSparseElementFromEntity(Entity, FTestFragment_SparseInt::StaticStruct());
+
+		// Note: EntityView doesn't get invalidated by sparse element changes (only archetype changes)
+		AITEST_TRUE(TEXT("View still valid after sparse removal"), EntityView.IsValid());
+
+		// But HasElement should return false now
+		AITEST_FALSE(TEXT("HasElement returns false after removal"),
+			EntityView.HasElement<FTestFragment_SparseInt>());
+
+		// And pointer should be null
+		AITEST_NULL(TEXT("GetSparseFragmentDataPtr returns null after removal"),
+			EntityView.GetSparseFragmentDataPtr<FTestFragment_SparseInt>());
+
+		return true;
+	}
+};
+IMPLEMENT_AI_INSTANT_TEST(FEntityView_SparseFragment_RemoveAndCheck, "System.Mass.Stress.EntityView.Sparse.RemoveAndCheck");
+
+
+struct FEntityView_HasElement_bIncludeSparse : FEntityTestBase
+{
+	virtual bool InstantTest() override
+	{
+		const FMassEntityHandle Entity = EntityManager->CreateEntity(IntsArchetype);
+
+		// Entity has regular Int fragment (from archetype) but not Float
+		AITEST_TRUE(TEXT("Entity has Int fragment in archetype"),
+			EntityManager->GetFragmentDataPtr<FTestFragment_Int>(Entity) != nullptr);
+		AITEST_FALSE(TEXT("Entity doesn't have Float fragment in archetype"),
+			EntityManager->GetFragmentDataPtr<FTestFragment_Float>(Entity) != nullptr);
+
+		// Add sparse fragment
+		EntityManager->AddSparseElementToEntity<FTestFragment_SparseFloat>(Entity);
+
+		FMassEntityView EntityView(*EntityManager, Entity);
+
+		// HasElement with bIncludeSparse=true should detect sparse fragment
+		AITEST_TRUE(TEXT("HasElement(true) detects sparse fragment"),
+			EntityView.HasElement<FTestFragment_SparseFloat>(UE::Mass::EIncludeSparseElements::Yes));
+
+		// HasElement with bIncludeSparse=false should NOT detect sparse fragment
+		AITEST_FALSE(TEXT("HasElement(false) doesn't detect sparse fragment"),
+			EntityView.HasElement<FTestFragment_SparseFloat>(UE::Mass::EIncludeSparseElements::No));
+
+		return true;
+	}
+};
+IMPLEMENT_AI_INSTANT_TEST(FEntityView_HasElement_bIncludeSparse, "System.Mass.Stress.EntityView.Sparse.HasElement.bIncludeSparse");
+
+
+struct FEntityView_SparseAcrossDifferentArchetypes : FEntityTestBase
+{
+	virtual bool InstantTest() override
+	{
+		const FMassEntityHandle EntityInts = EntityManager->CreateEntity(IntsArchetype);
+		const FMassEntityHandle EntityFloats = EntityManager->CreateEntity(FloatsArchetype);
+
+		// Add same sparse fragment to entities in different archetypes
+		EntityManager->AddSparseElementToEntity<FTestFragment_SparseInt>(EntityInts);
+		EntityManager->AddSparseElementToEntity<FTestFragment_SparseInt>(EntityFloats);
+
+		FMassEntityView ViewInts(*EntityManager, EntityInts);
+		FMassEntityView ViewFloats(*EntityManager, EntityFloats);
+
+		// Both should detect the sparse fragment
+		AITEST_TRUE(TEXT("Ints archetype entity has sparse fragment"),
+			ViewInts.HasElement<FTestFragment_SparseInt>());
+		AITEST_TRUE(TEXT("Floats archetype entity has sparse fragment"),
+			ViewFloats.HasElement<FTestFragment_SparseInt>());
+
+		// Both should provide access to the fragment
+		AITEST_NOT_NULL(TEXT("Ints entity can get sparse fragment"),
+			ViewInts.GetSparseFragmentDataPtr<FTestFragment_SparseInt>());
+		AITEST_NOT_NULL(TEXT("Floats entity can get sparse fragment"),
+			ViewFloats.GetSparseFragmentDataPtr<FTestFragment_SparseInt>());
+
+		return true;
+	}
+};
+IMPLEMENT_AI_INSTANT_TEST(FEntityView_SparseAcrossDifferentArchetypes, "System.Mass.Stress.EntityView.Sparse.AcrossDifferentArchetypes");
+
+
+struct FEntityView_MixedRegularAndSparseFragments : FEntityTestBase
+{
+	virtual bool InstantTest() override
+	{
+		const FMassEntityHandle Entity = EntityManager->CreateEntity(FloatsIntsArchetype);
+
+		// Entity has regular Float and Int fragments from archetype
+		// Add sparse fragment
+		EntityManager->AddSparseElementToEntity<FTestFragment_SparseInt>(Entity);
+
+		FMassEntityView EntityView(*EntityManager, Entity);
+
+		// Regular fragments should be accessible
+		FTestFragment_Float* FloatFragment = EntityView.GetFragmentDataPtr<FTestFragment_Float>();
+		FTestFragment_Int* IntFragment = EntityView.GetFragmentDataPtr<FTestFragment_Int>();
+		AITEST_NOT_NULL(TEXT("Can access regular Float fragment"), FloatFragment);
+		AITEST_NOT_NULL(TEXT("Can access regular Int fragment"), IntFragment);
+
+		// Sparse fragment should be accessible
+		FTestFragment_SparseInt* SparseFragment = EntityView.GetSparseFragmentDataPtr<FTestFragment_SparseInt>();
+		AITEST_NOT_NULL(TEXT("Can access sparse fragment"), SparseFragment);
+
+		// HasFragment should detect both types
+		AITEST_TRUE(TEXT("HasFragment detects regular Float"),
+			EntityView.HasFragment<FTestFragment_Float>());
+		AITEST_TRUE(TEXT("HasFragment detects regular Int"),
+			EntityView.HasFragment<FTestFragment_Int>());
+		AITEST_TRUE(TEXT("HasElement detects sparse Int"),
+			EntityView.HasElement<FTestFragment_SparseInt>());
+
+		return true;
+	}
+};
+IMPLEMENT_AI_INSTANT_TEST(FEntityView_MixedRegularAndSparseFragments, "System.Mass.Stress.EntityView.Sparse.MixedRegularAndSparse");
+
+
+struct FEntityView_TryMakeView_WithSparseElements : FEntityTestBase
+{
+	virtual bool InstantTest() override
+	{
+		const FMassEntityHandle ValidEntity = EntityManager->CreateEntity(IntsArchetype);
+		EntityManager->AddSparseElementToEntity<FTestFragment_SparseInt>(ValidEntity);
+
+		const FMassEntityHandle InvalidEntity; // Default constructor = invalid
+
+		// TryMakeView with valid entity
+		FMassEntityView ValidView = FMassEntityView::TryMakeView(*EntityManager, ValidEntity);
+		AITEST_TRUE(TEXT("TryMakeView succeeds for valid entity"), ValidView.IsSet());
+		AITEST_TRUE(TEXT("View detects sparse fragment"),
+			ValidView.HasElement<FTestFragment_SparseInt>());
+
+		// TryMakeView with invalid entity
+		FMassEntityView InvalidView = FMassEntityView::TryMakeView(*EntityManager, InvalidEntity);
+		AITEST_FALSE(TEXT("TryMakeView returns unset view for invalid entity"), InvalidView.IsSet());
+
+		return true;
+	}
+};
+IMPLEMENT_AI_INSTANT_TEST(FEntityView_TryMakeView_WithSparseElements, "System.Mass.Stress.EntityView.Sparse.TryMakeView");
+
+#endif // WITH_MASSENTITY_DEBUG
+} // UE::Mass::Test::EntityView
+
+UE_ENABLE_OPTIMIZATION_SHIP
+
+#undef LOCTEXT_NAMESPACE
diff --git a/Engine/Source/Developer/MassEntityTestSuite/Private/StressTests/MassStressRequirementsTest.cpp b/Engine/Source/Developer/MassEntityTestSuite/Private/StressTests/MassStressRequirementsTest.cpp
new file mode 100644
--- /dev/null
+++ b/Engine/Source/Developer/MassEntityTestSuite/Private/StressTests/MassStressRequirementsTest.cpp
@@ -0,0 +1,1234 @@
+// Copyright Epic Games, Inc. All Rights Reserved.
+
+#include "AITestsCommon.h"
+#include "MassEntityManager.h"
+#include "MassEntityTestTypes.h"
+#include "MassRequirements.h"
+#include "MassArchetypeTypes.h"
+
+#define LOCTEXT_NAMESPACE "MassTest"
+
+UE_DISABLE_OPTIMIZATION_SHIP
+
+namespace UE::Mass::Test::Requirements
+{
+#if WITH_MASSENTITY_DEBUG
+
+//----------------------------------------------------------------------//
+// FMassFragmentRequirements - Fragment Requirements Tests
+//----------------------------------------------------------------------//
+
+struct FRequirements_AddRequirement_All : FEntityTestBase
+{
+	virtual bool InstantTest() override
+	{
+		FMassFragmentRequirements Requirements(EntityManager);
+
+		// Add fragment with All presence
+		Requirements.AddRequirement<FTestFragment_Int>(EMassFragmentAccess::ReadOnly, EMassFragmentPresence::All);
+
+		// Verify fragment is in RequiredAllFragments
+		AITEST_TRUE("RequiredAllFragments should contain the added fragment",
+			Requirements.GetRequiredAllFragments().Contains<FTestFragment_Int>());
+		AITEST_FALSE("RequiredAnyFragments should not contain the fragment",
+			Requirements.GetRequiredAnyFragments().Contains<FTestFragment_Int>());
+		AITEST_FALSE("RequiredOptionalFragments should not contain the fragment",
+			Requirements.GetRequiredOptionalFragments().Contains<FTestFragment_Int>());
+		AITEST_FALSE("RequiredNoneFragments should not contain the fragment",
+			Requirements.GetRequiredNoneFragments().Contains<FTestFragment_Int>());
+
+		// Verify fragment requirements description
+		TConstArrayView<FMassFragmentRequirementDescription> FragmentRequirements = Requirements.GetFragmentRequirements();
+		AITEST_EQUAL("Should have one fragment requirement", FragmentRequirements.Num(), 1);
+		AITEST_EQUAL("Fragment type should match", FragmentRequirements[0].StructType, FTestFragment_Int::StaticStruct());
+		AITEST_EQUAL("Access mode should be ReadOnly", FragmentRequirements[0].AccessMode, EMassFragmentAccess::ReadOnly);
+		AITEST_EQUAL("Presence should be All", FragmentRequirements[0].Presence, EMassFragmentPresence::All);
+
+		return true;
+	}
+};
+IMPLEMENT_AI_INSTANT_TEST(FRequirements_AddRequirement_All, "System.Mass.Stress.Requirements.Fragment.AddRequirement.All");
+
+
+struct FRequirements_AddRequirement_Any : FEntityTestBase
+{
+	virtual bool InstantTest() override
+	{
+		FMassFragmentRequirements Requirements(EntityManager);
+
+		Requirements.AddRequirement<FTestFragment_Int>(EMassFragmentAccess::ReadWrite, EMassFragmentPresence::Any);
+
+		AITEST_FALSE("RequiredAllFragments should not contain the fragment",
+			Requirements.GetRequiredAllFragments().Contains<FTestFragment_Int>());
+		AITEST_TRUE("RequiredAnyFragments should contain the added fragment",
+			Requirements.GetRequiredAnyFragments().Contains<FTestFragment_Int>());
+		AITEST_FALSE("RequiredOptionalFragments should not contain the fragment",
+			Requirements.GetRequiredOptionalFragments().Contains<FTestFragment_Int>());
+		AITEST_FALSE("RequiredNoneFragments should not contain the fragment",
+			Requirements.GetRequiredNoneFragments().Contains<FTestFragment_Int>());
+
+		return true;
+	}
+};
+IMPLEMENT_AI_INSTANT_TEST(FRequirements_AddRequirement_Any, "System.Mass.Stress.Requirements.Fragment.AddRequirement.Any");
+
+
+struct FRequirements_AddRequirement_Optional : FEntityTestBase
+{
+	virtual bool InstantTest() override
+	{
+		FMassFragmentRequirements Requirements(EntityManager);
+
+		Requirements.AddRequirement<FTestFragment_Int>(EMassFragmentAccess::ReadOnly, EMassFragmentPresence::Optional);
+
+		AITEST_FALSE("RequiredAllFragments should not contain the fragment",
+			Requirements.GetRequiredAllFragments().Contains<FTestFragment_Int>());
+		AITEST_FALSE("RequiredAnyFragments should not contain the fragment",
+			Requirements.GetRequiredAnyFragments().Contains<FTestFragment_Int>());
+		AITEST_TRUE("RequiredOptionalFragments should contain the added fragment",
+			Requirements.GetRequiredOptionalFragments().Contains<FTestFragment_Int>());
+		AITEST_FALSE("RequiredNoneFragments should not contain the fragment",
+			Requirements.GetRequiredNoneFragments().Contains<FTestFragment_Int>());
+
+		return true;
+	}
+};
+IMPLEMENT_AI_INSTANT_TEST(FRequirements_AddRequirement_Optional, "System.Mass.Stress.Requirements.Fragment.AddRequirement.Optional");
+
+
+struct FRequirements_AddRequirement_None : FEntityTestBase
+{
+	virtual bool InstantTest() override
+	{
+		FMassFragmentRequirements Requirements(EntityManager);
+
+		Requirements.AddRequirement<FTestFragment_Int>(EMassFragmentAccess::None, EMassFragmentPresence::None);
+
+		AITEST_FALSE("RequiredAllFragments should not contain the fragment",
+			Requirements.GetRequiredAllFragments().Contains<FTestFragment_Int>());
+		AITEST_FALSE("RequiredAnyFragments should not contain the fragment",
+			Requirements.GetRequiredAnyFragments().Contains<FTestFragment_Int>());
+		AITEST_FALSE("RequiredOptionalFragments should not contain the fragment",
+			Requirements.GetRequiredOptionalFragments().Contains<FTestFragment_Int>());
+		AITEST_TRUE("RequiredNoneFragments should contain the added fragment",
+			Requirements.GetRequiredNoneFragments().Contains<FTestFragment_Int>());
+
+		// None presence should not add to FragmentRequirements array
+		AITEST_EQUAL("FragmentRequirements should be empty for None presence",
+			Requirements.GetFragmentRequirements().Num(), 0);
+
+		return true;
+	}
+};
+IMPLEMENT_AI_INSTANT_TEST(FRequirements_AddRequirement_None, "System.Mass.Stress.Requirements.Fragment.AddRequirement.None");
+
+
+struct FRequirements_AddRequirement_NonTemplate : FEntityTestBase
+{
+	virtual bool InstantTest() override
+	{
+		FMassFragmentRequirements Requirements(EntityManager);
+
+		// Use non-template version
+		Requirements.AddRequirement(FTestFragment_Float::StaticStruct(), EMassFragmentAccess::ReadWrite, EMassFragmentPresence::All);
+
+		AITEST_TRUE("RequiredAllFragments should contain the added fragment",
+			Requirements.GetRequiredAllFragments().Contains(FTestFragment_Float::StaticStruct()));
+
+		TConstArrayView<FMassFragmentRequirementDescription> FragmentRequirements = Requirements.GetFragmentRequirements();
+		AITEST_EQUAL("Should have one fragment requirement", FragmentRequirements.Num(), 1);
+		AITEST_EQUAL("Access mode should be ReadWrite", FragmentRequirements[0].AccessMode, EMassFragmentAccess::ReadWrite);
+
+		return true;
+	}
+};
+IMPLEMENT_AI_INSTANT_TEST(FRequirements_AddRequirement_NonTemplate, "System.Mass.Stress.Requirements.Fragment.AddRequirement.NonTemplate");
+
+
+struct FRequirements_AddRequirement_Multiple : FEntityTestBase
+{
+	virtual bool InstantTest() override
+	{
+		FMassFragmentRequirements Requirements(EntityManager);
+
+		Requirements.AddRequirement<FTestFragment_Int>(EMassFragmentAccess::ReadOnly, EMassFragmentPresence::All);
+		Requirements.AddRequirement<FTestFragment_Float>(EMassFragmentAccess::ReadWrite, EMassFragmentPresence::Optional);
+		Requirements.AddRequirement<FTestFragment_Bool>(EMassFragmentAccess::None, EMassFragmentPresence::None);
+
+		AITEST_TRUE("RequiredAllFragments should contain Int",
+			Requirements.GetRequiredAllFragments().Contains<FTestFragment_Int>());
+		AITEST_TRUE("RequiredOptionalFragments should contain Float",
+			Requirements.GetRequiredOptionalFragments().Contains<FTestFragment_Float>());
+		AITEST_TRUE("RequiredNoneFragments should contain Bool",
+			Requirements.GetRequiredNoneFragments().Contains<FTestFragment_Bool>());
+
+		// FragmentRequirements should have 2 entries (All and Optional, not None)
+		AITEST_EQUAL("Should have two fragment requirements (excluding None)",
+			Requirements.GetFragmentRequirements().Num(), 2);
+
+		return true;
+	}
+};
+IMPLEMENT_AI_INSTANT_TEST(FRequirements_AddRequirement_Multiple, "System.Mass.Stress.Requirements.Fragment.AddRequirement.Multiple");
+
+
+//----------------------------------------------------------------------//
+// FMassFragmentRequirements - Tag Requirements Tests
+//----------------------------------------------------------------------//
+
+struct FRequirements_AddTagRequirement_All : FEntityTestBase
+{
+	virtual bool InstantTest() override
+	{
+		FMassFragmentRequirements Requirements(EntityManager);
+
+		Requirements.AddTagRequirement<FTestTag_A>(EMassFragmentPresence::All);
+
+		AITEST_TRUE("RequiredAllTags should contain the added tag",
+			Requirements.GetRequiredAllTags().Contains<FTestTag_A>());
+		AITEST_FALSE("RequiredAnyTags should not contain the tag",
+			Requirements.GetRequiredAnyTags().Contains<FTestTag_A>());
+		AITEST_FALSE("RequiredNoneTags should not contain the tag",
+			Requirements.GetRequiredNoneTags().Contains<FTestTag_A>());
+		AITEST_FALSE("RequiredOptionalTags should not contain the tag",
+			Requirements.GetRequiredOptionalTags().Contains<FTestTag_A>());
+
+		return true;
+	}
+};
+IMPLEMENT_AI_INSTANT_TEST(FRequirements_AddTagRequirement_All, "System.Mass.Stress.Requirements.Tag.AddTagRequirement.All");
+
+
+struct FRequirements_AddTagRequirement_Any : FEntityTestBase
+{
+	virtual bool InstantTest() override
+	{
+		FMassFragmentRequirements Requirements(EntityManager);
+
+		Requirements.AddTagRequirement<FTestTag_A>(EMassFragmentPresence::Any);
+
+		AITEST_FALSE("RequiredAllTags should not contain the tag",
+			Requirements.GetRequiredAllTags().Contains<FTestTag_A>());
+		AITEST_TRUE("RequiredAnyTags should contain the added tag",
+			Requirements.GetRequiredAnyTags().Contains<FTestTag_A>());
+		AITEST_FALSE("RequiredNoneTags should not contain the tag",
+			Requirements.GetRequiredNoneTags().Contains<FTestTag_A>());
+		AITEST_FALSE("RequiredOptionalTags should not contain the tag",
+			Requirements.GetRequiredOptionalTags().Contains<FTestTag_A>());
+
+		return true;
+	}
+};
+IMPLEMENT_AI_INSTANT_TEST(FRequirements_AddTagRequirement_Any, "System.Mass.Stress.Requirements.Tag.AddTagRequirement.Any");
+
+
+struct FRequirements_AddTagRequirement_None : FEntityTestBase
+{
+	virtual bool InstantTest() override
+	{
+		FMassFragmentRequirements Requirements(EntityManager);
+
+		Requirements.AddTagRequirement<FTestTag_A>(EMassFragmentPresence::None);
+
+		AITEST_FALSE("RequiredAllTags should not contain the tag",
+			Requirements.GetRequiredAllTags().Contains<FTestTag_A>());
+		AITEST_FALSE("RequiredAnyTags should not contain the tag",
+			Requirements.GetRequiredAnyTags().Contains<FTestTag_A>());
+		AITEST_TRUE("RequiredNoneTags should contain the added tag",
+			Requirements.GetRequiredNoneTags().Contains<FTestTag_A>());
+		AITEST_FALSE("RequiredOptionalTags should not contain the tag",
+			Requirements.GetRequiredOptionalTags().Contains<FTestTag_A>());
+
+		return true;
+	}
+};
+IMPLEMENT_AI_INSTANT_TEST(FRequirements_AddTagRequirement_None, "System.Mass.Stress.Requirements.Tag.AddTagRequirement.None");
+
+
+struct FRequirements_AddTagRequirement_Optional : FEntityTestBase
+{
+	virtual bool InstantTest() override
+	{
+		FMassFragmentRequirements Requirements(EntityManager);
+
+		Requirements.AddTagRequirement<FTestTag_A>(EMassFragmentPresence::Optional);
+
+		AITEST_FALSE("RequiredAllTags should not contain the tag",
+			Requirements.GetRequiredAllTags().Contains<FTestTag_A>());
+		AITEST_FALSE("RequiredAnyTags should not contain the tag",
+			Requirements.GetRequiredAnyTags().Contains<FTestTag_A>());
+		AITEST_FALSE("RequiredNoneTags should not contain the tag",
+			Requirements.GetRequiredNoneTags().Contains<FTestTag_A>());
+		AITEST_TRUE("RequiredOptionalTags should contain the added tag",
+			Requirements.GetRequiredOptionalTags().Contains<FTestTag_A>());
+
+		return true;
+	}
+};
+IMPLEMENT_AI_INSTANT_TEST(FRequirements_AddTagRequirement_Optional, "System.Mass.Stress.Requirements.Tag.AddTagRequirement.Optional");
+
+
+struct FRequirements_AddTagRequirement_NonTemplate : FEntityTestBase
+{
+	virtual bool InstantTest() override
+	{
+		FMassFragmentRequirements Requirements(EntityManager);
+
+		Requirements.AddTagRequirement(FTestTag_B::StaticStruct(), EMassFragmentPresence::All);
+
+		AITEST_TRUE("RequiredAllTags should contain the added tag",
+			Requirements.GetRequiredAllTags().Contains(FTestTag_B::StaticStruct()));
+
+		return true;
+	}
+};
+IMPLEMENT_AI_INSTANT_TEST(FRequirements_AddTagRequirement_NonTemplate, "System.Mass.Stress.Requirements.Tag.AddTagRequirement.NonTemplate");
+
+
+struct FRequirements_AddTagRequirements_Batch : FEntityTestBase
+{
+	virtual bool InstantTest() override
+	{
+		FMassFragmentRequirements Requirements(EntityManager);
+
+		FMassTagBitSet TagBitSet;
+		TagBitSet.Add<FTestTag_A>();
+		TagBitSet.Add<FTestTag_B>();
+
+		// Add batch of tags with All presence
+		Requirements.AddTagRequirements<EMassFragmentPresence::All>(TagBitSet);
+
+		AITEST_TRUE("RequiredAllTags should contain Tag A",
+			Requirements.GetRequiredAllTags().Contains<FTestTag_A>());
+		AITEST_TRUE("RequiredAllTags should contain Tag B",
+			Requirements.GetRequiredAllTags().Contains<FTestTag_B>());
+
+		// Add another batch with None presence
+		FMassTagBitSet NoneTagBitSet;
+		NoneTagBitSet.Add<FTestTag_C>();
+		Requirements.AddTagRequirements<EMassFragmentPresence::None>(NoneTagBitSet);
+
+		AITEST_TRUE("RequiredNoneTags should contain Tag C",
+			Requirements.GetRequiredNoneTags().Contains<FTestTag_C>());
+
+		return true;
+	}
+};
+IMPLEMENT_AI_INSTANT_TEST(FRequirements_AddTagRequirements_Batch, "System.Mass.Stress.Requirements.Tag.AddTagRequirements.Batch");
+
+
+struct FRequirements_ClearTagRequirements : FEntityTestBase
+{
+	virtual bool InstantTest() override
+	{
+		FMassFragmentRequirements Requirements(EntityManager);
+
+		// Add tags to various categories
+		Requirements.AddTagRequirement<FTestTag_A>(EMassFragmentPresence::All);
+		Requirements.AddTagRequirement<FTestTag_B>(EMassFragmentPresence::Any);
+		Requirements.AddTagRequirement<FTestTag_C>(EMassFragmentPresence::None);
+
+		// Verify tags are present
+		AITEST_TRUE("Tag A should be in AllTags", Requirements.GetRequiredAllTags().Contains<FTestTag_A>());
+		AITEST_TRUE("Tag B should be in AnyTags", Requirements.GetRequiredAnyTags().Contains<FTestTag_B>());
+		AITEST_TRUE("Tag C should be in NoneTags", Requirements.GetRequiredNoneTags().Contains<FTestTag_C>());
+
+		// Clear Tag A and B
+		FMassTagBitSet TagsToClear;
+		TagsToClear.Add<FTestTag_A>();
+		TagsToClear.Add<FTestTag_B>();
+		Requirements.ClearTagRequirements(TagsToClear);
+
+		// Verify cleared tags are gone
+		AITEST_FALSE("Tag A should be cleared from AllTags", Requirements.GetRequiredAllTags().Contains<FTestTag_A>());
+		AITEST_FALSE("Tag B should be cleared from AnyTags", Requirements.GetRequiredAnyTags().Contains<FTestTag_B>());
+
+		// Tag C should still be present (wasn't in clear set)
+		AITEST_TRUE("Tag C should still be in NoneTags", Requirements.GetRequiredNoneTags().Contains<FTestTag_C>());
+
+		return true;
+	}
+};
+IMPLEMENT_AI_INSTANT_TEST(FRequirements_ClearTagRequirements, "System.Mass.Stress.Requirements.Tag.ClearTagRequirements");
+
+
+struct FRequirements_AddTagRequirement_Multiple : FEntityTestBase
+{
+	virtual bool InstantTest() override
+	{
+		FMassFragmentRequirements Requirements(EntityManager);
+
+		Requirements.AddTagRequirement<FTestTag_A>(EMassFragmentPresence::All);
+		Requirements.AddTagRequirement<FTestTag_B>(EMassFragmentPresence::Any);
+		Requirements.AddTagRequirement<FTestTag_C>(EMassFragmentPresence::None);
+		Requirements.AddTagRequirement<FTestTag_D>(EMassFragmentPresence::Optional);
+
+		AITEST_TRUE("Tag A in AllTags", Requirements.GetRequiredAllTags().Contains<FTestTag_A>());
+		AITEST_TRUE("Tag B in AnyTags", Requirements.GetRequiredAnyTags().Contains<FTestTag_B>());
+		AITEST_TRUE("Tag C in NoneTags", Requirements.GetRequiredNoneTags().Contains<FTestTag_C>());
+		AITEST_TRUE("Tag D in OptionalTags", Requirements.GetRequiredOptionalTags().Contains<FTestTag_D>());
+
+		return true;
+	}
+};
+IMPLEMENT_AI_INSTANT_TEST(FRequirements_AddTagRequirement_Multiple, "System.Mass.Stress.Requirements.Tag.AddTagRequirement.Multiple");
+
+
+//----------------------------------------------------------------------//
+// FMassFragmentRequirements - Sparse Element Requirements Tests
+// These tests use the unified API (AddRequirement/AddTagRequirement) with sparse elements
+// to verify that the functions automatically detect and route sparse elements correctly
+//----------------------------------------------------------------------//
+
+struct FRequirements_AddRequirement_SparseFragment_All : FEntityTestBase
+{
+	virtual bool InstantTest() override
+	{
+		FMassFragmentRequirements Requirements(EntityManager);
+
+		// Use AddRequirement with a sparse fragment - should route to sparse storage
+		Requirements.AddRequirement<FTestFragment_SparseInt>(EMassFragmentAccess::ReadOnly, EMassFragmentPresence::All);
+
+		// Should be in sparse elements, NOT regular fragments
+		AITEST_TRUE("RequiredAllSparseElements should contain the sparse fragment",
+			Requirements.GetRequiredAllSparseElements().Contains(FTestFragment_SparseInt::StaticStruct()));
+		AITEST_FALSE("RequiredAllFragments should NOT contain the sparse fragment",
+			Requirements.GetRequiredAllFragments().Contains(FTestFragment_SparseInt::StaticStruct()));
+		AITEST_FALSE("RequiredAnySparseElements should not contain the sparse fragment",
+			Requirements.GetRequiredAnySparseElements().Contains(FTestFragment_SparseInt::StaticStruct()));
+		AITEST_FALSE("RequiredOptionalSparseElements should not contain the sparse fragment",
+			Requirements.GetRequiredOptionalSparseElements().Contains(FTestFragment_SparseInt::StaticStruct()));
+		AITEST_FALSE("RequiredNoneSparseElements should not contain the sparse fragment",
+			Requirements.GetRequiredNoneSparseElements().Contains(FTestFragment_SparseInt::StaticStruct()));
+
+		Requirements.CheckValidity(); // force re-caching of "HasXRequirements" cached properties
+		AITEST_TRUE("HasSparseRequirements should return true", Requirements.HasSparseRequirements());
+
+		return true;
+	}
+};
+IMPLEMENT_AI_INSTANT_TEST(FRequirements_AddRequirement_SparseFragment_All, "System.Mass.Stress.Requirements.Sparse.AddRequirement.SparseFragment.All");
+
+
+struct FRequirements_AddRequirement_SparseFragment_Any : FEntityTestBase
+{
+	virtual bool InstantTest() override
+	{
+		FMassFragmentRequirements Requirements(EntityManager);
+
+		// Use AddRequirement with a sparse fragment - should route to sparse storage
+		Requirements.AddRequirement<FTestFragment_SparseInt>(EMassFragmentAccess::ReadOnly, EMassFragmentPresence::Any);
+
+		AITEST_FALSE("RequiredAllSparseElements should not contain the sparse fragment",
+			Requirements.GetRequiredAllSparseElements().Contains(FTestFragment_SparseInt::StaticStruct()));
+		AITEST_TRUE("RequiredAnySparseElements should contain the sparse fragment",
+			Requirements.GetRequiredAnySparseElements().Contains(FTestFragment_SparseInt::StaticStruct()));
+		AITEST_FALSE("RequiredAnyFragments should NOT contain the sparse fragment",
+			Requirements.GetRequiredAnyFragments().Contains(FTestFragment_SparseInt::StaticStruct()));
+
+		return true;
+	}
+};
+IMPLEMENT_AI_INSTANT_TEST(FRequirements_AddRequirement_SparseFragment_Any, "System.Mass.Stress.Requirements.Sparse.AddRequirement.SparseFragment.Any");
+
+
+struct FRequirements_AddRequirement_SparseFragment_Optional : FEntityTestBase
+{
+	virtual bool InstantTest() override
+	{
+		FMassFragmentRequirements Requirements(EntityManager);
+
+		// Use AddRequirement with a sparse fragment - should route to sparse storage
+		Requirements.AddRequirement<FTestFragment_SparseFloat>(EMassFragmentAccess::ReadOnly, EMassFragmentPresence::Optional);
+
+		AITEST_FALSE("RequiredAllSparseElements should not contain the sparse fragment",
+			Requirements.GetRequiredAllSparseElements().Contains(FTestFragment_SparseFloat::StaticStruct()));
+		AITEST_TRUE("RequiredOptionalSparseElements should contain the sparse fragment",
+			Requirements.GetRequiredOptionalSparseElements().Contains(FTestFragment_SparseFloat::StaticStruct()));
+		AITEST_FALSE("RequiredOptionalFragments should NOT contain the sparse fragment",
+			Requirements.GetRequiredOptionalFragments().Contains(FTestFragment_SparseFloat::StaticStruct()));
+
+		return true;
+	}
+};
+IMPLEMENT_AI_INSTANT_TEST(FRequirements_AddRequirement_SparseFragment_Optional, "System.Mass.Stress.Requirements.Sparse.AddRequirement.SparseFragment.Optional");
+
+
+struct FRequirements_AddRequirement_SparseFragment_None : FEntityTestBase
+{
+	virtual bool InstantTest() override
+	{
+		FMassFragmentRequirements Requirements(EntityManager);
+
+		// Use AddRequirement with a sparse fragment - should route to sparse storage
+		Requirements.AddRequirement<FTestFragment_SparseInt>(EMassFragmentAccess::None, EMassFragmentPresence::None);
+
+		AITEST_TRUE("RequiredNoneSparseElements should contain the sparse fragment",
+			Requirements.GetRequiredNoneSparseElements().Contains(FTestFragment_SparseInt::StaticStruct()));
+		AITEST_FALSE("RequiredNoneFragments should NOT contain the sparse fragment",
+			Requirements.GetRequiredNoneFragments().Contains(FTestFragment_SparseInt::StaticStruct()));
+
+		return true;
+	}
+};
+IMPLEMENT_AI_INSTANT_TEST(FRequirements_AddRequirement_SparseFragment_None, "System.Mass.Stress.Requirements.Sparse.AddRequirement.SparseFragment.None");
+
+
+struct FRequirements_AddTagRequirement_SparseTag_All : FEntityTestBase
+{
+	virtual bool InstantTest() override
+	{
+		FMassFragmentRequirements Requirements(EntityManager);
+
+		// Use AddTagRequirement with a sparse tag - should route to sparse storage
+		Requirements.AddTagRequirement<FTestTag_SparseA>(EMassFragmentPresence::All);
+
+		AITEST_TRUE("RequiredAllSparseElements should contain the sparse tag",
+			Requirements.GetRequiredAllSparseElements().Contains(FTestTag_SparseA::StaticStruct()));
+		AITEST_FALSE("RequiredAllTags should NOT contain the sparse tag",
+			Requirements.GetRequiredAllTags().Contains(FTestTag_SparseA::StaticStruct()));
+
+		Requirements.CheckValidity(); // force re-caching of "HasXRequirements" cached properties
+		AITEST_TRUE("HasSparseRequirements should return true", Requirements.HasSparseRequirements());
+
+		return true;
+	}
+};
+IMPLEMENT_AI_INSTANT_TEST(FRequirements_AddTagRequirement_SparseTag_All, "System.Mass.Stress.Requirements.Sparse.AddTagRequirement.SparseTag.All");
+
+
+struct FRequirements_AddTagRequirement_SparseTag_Any : FEntityTestBase
+{
+	virtual bool InstantTest() override
+	{
+		FMassFragmentRequirements Requirements(EntityManager);
+
+		// Use AddTagRequirement with a sparse tag - should route to sparse storage
+		Requirements.AddTagRequirement<FTestTag_SparseA>(EMassFragmentPresence::Any);
+
+		AITEST_TRUE("RequiredAnySparseElements should contain the sparse tag",
+			Requirements.GetRequiredAnySparseElements().Contains(FTestTag_SparseA::StaticStruct()));
+		AITEST_FALSE("RequiredAnyTags should NOT contain the sparse tag",
+			Requirements.GetRequiredAnyTags().Contains(FTestTag_SparseA::StaticStruct()));
+
+		return true;
+	}
+};
+IMPLEMENT_AI_INSTANT_TEST(FRequirements_AddTagRequirement_SparseTag_Any, "System.Mass.Stress.Requirements.Sparse.AddTagRequirement.SparseTag.Any");
+
+
+struct FRequirements_AddTagRequirement_SparseTag_Optional : FEntityTestBase
+{
+	virtual bool InstantTest() override
+	{
+		FMassFragmentRequirements Requirements(EntityManager);
+
+		// Use AddTagRequirement with a sparse tag - should route to sparse storage
+		Requirements.AddTagRequirement<FTestTag_SparseB>(EMassFragmentPresence::Optional);
+
+		AITEST_TRUE("RequiredOptionalSparseElements should contain the sparse tag",
+			Requirements.GetRequiredOptionalSparseElements().Contains(FTestTag_SparseB::StaticStruct()));
+		AITEST_FALSE("RequiredOptionalTags should NOT contain the sparse tag",
+			Requirements.GetRequiredOptionalTags().Contains(FTestTag_SparseB::StaticStruct()));
+
+		return true;
+	}
+};
+IMPLEMENT_AI_INSTANT_TEST(FRequirements_AddTagRequirement_SparseTag_Optional, "System.Mass.Stress.Requirements.Sparse.AddTagRequirement.SparseTag.Optional");
+
+
+struct FRequirements_AddTagRequirement_SparseTag_None : FEntityTestBase
+{
+	virtual bool InstantTest() override
+	{
+		FMassFragmentRequirements Requirements(EntityManager);
+
+		// Use AddTagRequirement with a sparse tag - should route to sparse storage
+		Requirements.AddTagRequirement<FTestTag_SparseA>(EMassFragmentPresence::None);
+
+		AITEST_TRUE("RequiredNoneSparseElements should contain the sparse tag",
+			Requirements.GetRequiredNoneSparseElements().Contains(FTestTag_SparseA::StaticStruct()));
+		AITEST_FALSE("RequiredNoneTags should NOT contain the sparse tag",
+			Requirements.GetRequiredNoneTags().Contains(FTestTag_SparseA::StaticStruct()));
+
+		return true;
+	}
+};
+IMPLEMENT_AI_INSTANT_TEST(FRequirements_AddTagRequirement_SparseTag_None, "System.Mass.Stress.Requirements.Sparse.AddTagRequirement.SparseTag.None");
+
+
+struct FRequirements_AddRequirement_SparseFragment_NonTemplate : FEntityTestBase
+{
+	virtual bool InstantTest() override
+	{
+		FMassFragmentRequirements Requirements(EntityManager);
+
+		// Use non-template AddRequirement with sparse fragment
+		Requirements.AddRequirement(FTestFragment_SparseInt::StaticStruct(), EMassFragmentAccess::ReadOnly, EMassFragmentPresence::All);
+
+		AITEST_TRUE("RequiredAllSparseElements should contain the sparse fragment",
+			Requirements.GetRequiredAllSparseElements().Contains(FTestFragment_SparseInt::StaticStruct()));
+		AITEST_FALSE("RequiredAllFragments should NOT contain the sparse fragment",
+			Requirements.GetRequiredAllFragments().Contains(FTestFragment_SparseInt::StaticStruct()));
+
+		return true;
+	}
+};
+IMPLEMENT_AI_INSTANT_TEST(FRequirements_AddRequirement_SparseFragment_NonTemplate, "System.Mass.Stress.Requirements.Sparse.AddRequirement.SparseFragment.NonTemplate");
+
+
+struct FRequirements_AddTagRequirement_SparseTag_NonTemplate : FEntityTestBase
+{
+	virtual bool InstantTest() override
+	{
+		FMassFragmentRequirements Requirements(EntityManager);
+
+		// Use non-template AddTagRequirement with sparse tag
+		Requirements.AddTagRequirement(FTestTag_SparseB::StaticStruct(), EMassFragmentPresence::Any);
+
+		AITEST_TRUE("RequiredAnySparseElements should contain the sparse tag",
+			Requirements.GetRequiredAnySparseElements().Contains(FTestTag_SparseB::StaticStruct()));
+		AITEST_FALSE("RequiredAnyTags should NOT contain the sparse tag",
+			Requirements.GetRequiredAnyTags().Contains(FTestTag_SparseB::StaticStruct()));
+
+		return true;
+	}
+};
+IMPLEMENT_AI_INSTANT_TEST(FRequirements_AddTagRequirement_SparseTag_NonTemplate, "System.Mass.Stress.Requirements.Sparse.AddTagRequirement.SparseTag.NonTemplate");
+
+
+struct FRequirements_UnifiedAPI_MixedSparseElements : FEntityTestBase
+{
+	virtual bool InstantTest() override
+	{
+		FMassFragmentRequirements Requirements(EntityManager);
+
+		// Mix of sparse fragments and sparse tags using unified API
+		Requirements.AddRequirement<FTestFragment_SparseInt>(EMassFragmentAccess::ReadOnly, EMassFragmentPresence::All);
+		Requirements.AddRequirement<FTestFragment_SparseFloat>(EMassFragmentAccess::ReadWrite, EMassFragmentPresence::Optional);
+		Requirements.AddTagRequirement<FTestTag_SparseA>(EMassFragmentPresence::Any);
+		Requirements.AddTagRequirement<FTestTag_SparseB>(EMassFragmentPresence::None);
+
+		// All should be routed to sparse storage
+		AITEST_TRUE("Sparse int in AllSparseElements",
+			Requirements.GetRequiredAllSparseElements().Contains(FTestFragment_SparseInt::StaticStruct()));
+		AITEST_TRUE("Sparse float in OptionalSparseElements",
+			Requirements.GetRequiredOptionalSparseElements().Contains(FTestFragment_SparseFloat::StaticStruct()));
+		AITEST_TRUE("Sparse tag A in AnySparseElements",
+			Requirements.GetRequiredAnySparseElements().Contains(FTestTag_SparseA::StaticStruct()));
+		AITEST_TRUE("Sparse tag B in NoneSparseElements",
+			Requirements.GetRequiredNoneSparseElements().Contains(FTestTag_SparseB::StaticStruct()));
+
+		// None should be in regular storage
+		AITEST_FALSE("Sparse int NOT in AllFragments",
+			Requirements.GetRequiredAllFragments().Contains(FTestFragment_SparseInt::StaticStruct()));
+		AITEST_FALSE("Sparse float NOT in OptionalFragments",
+			Requirements.GetRequiredOptionalFragments().Contains(FTestFragment_SparseFloat::StaticStruct()));
+		AITEST_FALSE("Sparse tag A NOT in AnyTags",
+			Requirements.GetRequiredAnyTags().Contains(FTestTag_SparseA::StaticStruct()));
+		AITEST_FALSE("Sparse tag B NOT in NoneTags",
+			Requirements.GetRequiredNoneTags().Contains(FTestTag_SparseB::StaticStruct()));
+
+		Requirements.CheckValidity(); // force re-caching of "HasXRequirements" cached properties
+		AITEST_TRUE("HasSparseRequirements should return true", Requirements.HasSparseRequirements());
+
+		return true;
+	}
+};
+IMPLEMENT_AI_INSTANT_TEST(FRequirements_UnifiedAPI_MixedSparseElements, "System.Mass.Stress.Requirements.Sparse.UnifiedAPI.Mixed");
+
+
+struct FRequirements_UnifiedAPI_MixedRegularAndSparse : FEntityTestBase
+{
+	virtual bool InstantTest() override
+	{
+		FMassFragmentRequirements Requirements(EntityManager);
+
+		// Add regular fragments
+		Requirements.AddRequirement<FTestFragment_Int>(EMassFragmentAccess::ReadOnly, EMassFragmentPresence::All);
+		Requirements.AddRequirement<FTestFragment_Float>(EMassFragmentAccess::ReadWrite, EMassFragmentPresence::Optional);
+		// Add regular tags
+		Requirements.AddTagRequirement<FTestTag_A>(EMassFragmentPresence::All);
+		Requirements.AddTagRequirement<FTestTag_B>(EMassFragmentPresence::None);
+		// Add sparse fragments using unified API
+		Requirements.AddRequirement<FTestFragment_SparseInt>(EMassFragmentAccess::ReadOnly, EMassFragmentPresence::All);
+		Requirements.AddRequirement<FTestFragment_SparseFloat>(EMassFragmentAccess::ReadWrite, EMassFragmentPresence::Any);
+		// Add sparse tags using unified API
+		Requirements.AddTagRequirement<FTestTag_SparseA>(EMassFragmentPresence::Optional);
+		Requirements.AddTagRequirement<FTestTag_SparseB>(EMassFragmentPresence::None);
+
+		// Regular elements should be in regular storage
+		AITEST_TRUE("Int fragment in AllFragments", Requirements.GetRequiredAllFragments().Contains<FTestFragment_Int>());
+		AITEST_TRUE("Float fragment in OptionalFragments", Requirements.GetRequiredOptionalFragments().Contains<FTestFragment_Float>());
+		AITEST_TRUE("Tag A in AllTags", Requirements.GetRequiredAllTags().Contains<FTestTag_A>());
+		AITEST_TRUE("Tag B in NoneTags", Requirements.GetRequiredNoneTags().Contains<FTestTag_B>());
+
+		// Sparse elements should be in sparse storage
+		AITEST_TRUE("Sparse int in AllSparseElements",
+			Requirements.GetRequiredAllSparseElements().Contains(FTestFragment_SparseInt::StaticStruct()));
+		AITEST_TRUE("Sparse float in AnySparseElements",
+			Requirements.GetRequiredAnySparseElements().Contains(FTestFragment_SparseFloat::StaticStruct()));
+		AITEST_TRUE("Sparse tag A in OptionalSparseElements",
+			Requirements.GetRequiredOptionalSparseElements().Contains(FTestTag_SparseA::StaticStruct()));
+		AITEST_TRUE("Sparse tag B in NoneSparseElements",
+			Requirements.GetRequiredNoneSparseElements().Contains(FTestTag_SparseB::StaticStruct()));
+
+		// Sparse elements should NOT be in regular storage
+		AITEST_FALSE("Sparse int NOT in regular AllFragments",
+			Requirements.GetRequiredAllFragments().Contains(FTestFragment_SparseInt::StaticStruct()));
+		AITEST_FALSE("Sparse float NOT in regular AnyFragments",
+			Requirements.GetRequiredAnyFragments().Contains(FTestFragment_SparseFloat::StaticStruct()));
+		AITEST_FALSE("Sparse tag A NOT in regular OptionalTags",
+			Requirements.GetRequiredOptionalTags().Contains(FTestTag_SparseA::StaticStruct()));
+		AITEST_FALSE("Sparse tag B NOT in regular NoneTags",
+			Requirements.GetRequiredNoneTags().Contains(FTestTag_SparseB::StaticStruct()));
+
+		Requirements.CheckValidity(); // force re-caching of "HasXRequirements" cached properties
+		AITEST_TRUE("HasSparseRequirements should return true", Requirements.HasSparseRequirements());
+		AITEST_TRUE("HasPositiveRequirements should return true", Requirements.HasPositiveRequirements());
+		AITEST_TRUE("HasNegativeRequirements should return true", Requirements.HasNegativeRequirements());
+		AITEST_TRUE("HasOptionalRequirements should return true", Requirements.HasOptionalRequirements());
+
+		return true;
+	}
+};
+IMPLEMENT_AI_INSTANT_TEST(FRequirements_UnifiedAPI_MixedRegularAndSparse, "System.Mass.Stress.Requirements.Sparse.UnifiedAPI.MixedRegularAndSparse");
+
+
+//----------------------------------------------------------------------//
+// FMassFragmentRequirements - State and Validation Tests
+//----------------------------------------------------------------------//
+
+struct FRequirements_IsEmpty : FEntityTestBase
+{
+	virtual bool InstantTest() override
+	{
+		FMassFragmentRequirements Requirements(EntityManager);
+
+		AITEST_TRUE("Newly created requirements should be empty", Requirements.IsEmpty());
+
+		Requirements.AddRequirement<FTestFragment_Int>(EMassFragmentAccess::ReadOnly, EMassFragmentPresence::All);
+		Requirements.CheckValidity(); // force re-caching of "HasXRequirements" cached properties
+
+		AITEST_FALSE("Requirements with fragment should not be empty", Requirements.IsEmpty());
+
+		return true;
+	}
+};
+IMPLEMENT_AI_INSTANT_TEST(FRequirements_IsEmpty, "System.Mass.Stress.Requirements.State.IsEmpty");
+
+
+struct FRequirements_HasPositiveRequirements : FEntityTestBase
+{
+	virtual bool InstantTest() override
+	{
+		FMassFragmentRequirements Requirements(EntityManager);
+
+		// Initially no positive requirements
+		AITEST_FALSE("Empty requirements have no positive requirements", Requirements.HasPositiveRequirements());
+
+		// Add "All" fragment - this is a positive requirement
+		Requirements.AddRequirement<FTestFragment_Int>(EMassFragmentAccess::ReadOnly, EMassFragmentPresence::All);
+		Requirements.CheckValidity(); // force re-caching of "HasXRequirements" cached properties
+
+		AITEST_TRUE("Requirements with All fragment have positive requirements", Requirements.HasPositiveRequirements());
+
+		return true;
+	}
+};
+IMPLEMENT_AI_INSTANT_TEST(FRequirements_HasPositiveRequirements, "System.Mass.Stress.Requirements.State.HasPositiveRequirements");
+
+
+struct FRequirements_HasNegativeRequirements : FEntityTestBase
+{
+	virtual bool InstantTest() override
+	{
+		FMassFragmentRequirements Requirements(EntityManager);
+
+		AITEST_FALSE("Empty requirements have no negative requirements", Requirements.HasNegativeRequirements());
+
+		// Add "None" fragment - this is a negative requirement
+		Requirements.AddRequirement<FTestFragment_Int>(EMassFragmentAccess::None, EMassFragmentPresence::None);
+		Requirements.CheckValidity(); // force re-caching of "HasXRequirements" cached properties
+
+		AITEST_TRUE("Requirements with None fragment have negative requirements", Requirements.HasNegativeRequirements());
+
+		return true;
+	}
+};
+IMPLEMENT_AI_INSTANT_TEST(FRequirements_HasNegativeRequirements, "System.Mass.Stress.Requirements.State.HasNegativeRequirements");
+
+
+struct FRequirements_HasOptionalRequirements : FEntityTestBase
+{
+	virtual bool InstantTest() override
+	{
+		FMassFragmentRequirements Requirements(EntityManager);
+
+		AITEST_FALSE("Empty requirements have no optional requirements", Requirements.HasOptionalRequirements());
+
+		// Add "Optional" fragment
+		Requirements.AddRequirement<FTestFragment_Int>(EMassFragmentAccess::ReadOnly, EMassFragmentPresence::Optional);
+		Requirements.CheckValidity(); // force re-caching of "HasXRequirements" cached properties
+
+		AITEST_TRUE("Requirements with Optional fragment have optional requirements", Requirements.HasOptionalRequirements());
+
+		return true;
+	}
+};
+IMPLEMENT_AI_INSTANT_TEST(FRequirements_HasOptionalRequirements, "System.Mass.Stress.Requirements.State.HasOptionalRequirements");
+
+
+struct FRequirements_HasSparseRequirements : FEntityTestBase
+{
+	virtual bool InstantTest() override
+	{
+		FMassFragmentRequirements Requirements(EntityManager);
+
+		AITEST_FALSE("Empty requirements have no sparse requirements", Requirements.HasSparseRequirements());
+
+		// Use unified API - AddRequirement should detect sparse element and route appropriately
+		Requirements.AddRequirement<FTestFragment_SparseInt>(EMassFragmentAccess::ReadOnly, EMassFragmentPresence::All);
+		Requirements.CheckValidity(); // force re-caching of "HasXRequirements" cached properties
+
+		AITEST_TRUE("Requirements with sparse element have sparse requirements", Requirements.HasSparseRequirements());
+
+		return true;
+	}
+};
+IMPLEMENT_AI_INSTANT_TEST(FRequirements_HasSparseRequirements, "System.Mass.Stress.Requirements.State.HasSparseRequirements");
+
+
+struct FRequirements_Reset : FEntityTestBase
+{
+	virtual bool InstantTest() override
+	{
+		FMassFragmentRequirements Requirements(EntityManager);
+
+		// Add various requirements - use unified API for sparse element
+		Requirements.AddRequirement<FTestFragment_Int>(EMassFragmentAccess::ReadOnly, EMassFragmentPresence::All);
+		Requirements.AddTagRequirement<FTestTag_A>(EMassFragmentPresence::All);
+		Requirements.AddRequirement<FTestFragment_SparseInt>(EMassFragmentAccess::ReadOnly, EMassFragmentPresence::All);
+		Requirements.CheckValidity(); // force re-caching of "HasXRequirements" cached properties
+
+		AITEST_FALSE("Requirements should not be empty before reset", Requirements.IsEmpty());
+
+		Requirements.Reset();
+		Requirements.CheckValidity(); // force re-caching of "HasXRequirements" cached properties
+
+		AITEST_TRUE("Requirements should be empty after reset", Requirements.IsEmpty());
+		AITEST_FALSE("Should not have positive requirements after reset", Requirements.HasPositiveRequirements());
+		AITEST_FALSE("Should not have sparse requirements after reset", Requirements.HasSparseRequirements());
+
+		return true;
+	}
+};
+IMPLEMENT_AI_INSTANT_TEST(FRequirements_Reset, "System.Mass.Stress.Requirements.State.Reset");
+
+
+//----------------------------------------------------------------------//
+// FMassFragmentRequirements - Archetype Matching Tests
+//----------------------------------------------------------------------//
+
+struct FRequirements_DoesArchetypeMatchRequirements_AllFragments : FEntityTestBase
+{
+	virtual bool InstantTest() override
+	{
+		FMassFragmentRequirements Requirements(EntityManager);
+		Requirements.AddRequirement<FTestFragment_Int>(EMassFragmentAccess::ReadOnly, EMassFragmentPresence::All);
+		Requirements.AddRequirement<FTestFragment_Float>(EMassFragmentAccess::ReadOnly, EMassFragmentPresence::All);
+		Requirements.CheckValidity(); // force re-caching of "HasXRequirements" cached properties
+
+		// IntsArchetype has FTestFragment_Int but not FTestFragment_Float
+		AITEST_FALSE("IntsArchetype should not match (missing Float)",
+			Requirements.DoesArchetypeMatchRequirements(IntsArchetype));
+
+		// FloatsIntsArchetype has both
+		AITEST_TRUE("FloatsIntsArchetype should match",
+			Requirements.DoesArchetypeMatchRequirements(FloatsIntsArchetype));
+
+		return true;
+	}
+};
+IMPLEMENT_AI_INSTANT_TEST(FRequirements_DoesArchetypeMatchRequirements_AllFragments, "System.Mass.Stress.Requirements.Match.AllFragments");
+
+
+struct FRequirements_DoesArchetypeMatchRequirements_AnyFragments : FEntityTestBase
+{
+	virtual bool InstantTest() override
+	{
+		FMassFragmentRequirements Requirements(EntityManager);
+		Requirements.AddRequirement<FTestFragment_Int>(EMassFragmentAccess::ReadOnly, EMassFragmentPresence::Any);
+		Requirements.AddRequirement<FTestFragment_Float>(EMassFragmentAccess::ReadOnly, EMassFragmentPresence::Any);
+		Requirements.CheckValidity(); // force re-caching of "HasXRequirements" cached properties
+
+		// IntsArchetype has Int
+		AITEST_TRUE("IntsArchetype should match (has Int)",
+			Requirements.DoesArchetypeMatchRequirements(IntsArchetype));
+
+		// FloatsArchetype has Float
+		AITEST_TRUE("FloatsArchetype should match (has Float)",
+			Requirements.DoesArchetypeMatchRequirements(FloatsArchetype));
+
+		// EmptyArchetype has neither
+		AITEST_FALSE("EmptyArchetype should not match (has neither)",
+			Requirements.DoesArchetypeMatchRequirements(EmptyArchetype));
+
+		return true;
+	}
+};
+IMPLEMENT_AI_INSTANT_TEST(FRequirements_DoesArchetypeMatchRequirements_AnyFragments, "System.Mass.Stress.Requirements.Match.AnyFragments");
+
+
+struct FRequirements_DoesArchetypeMatchRequirements_NoneFragments : FEntityTestBase
+{
+	virtual bool InstantTest() override
+	{
+		FMassFragmentRequirements Requirements(EntityManager);
+		Requirements.AddRequirement<FTestFragment_Int>(EMassFragmentAccess::ReadOnly, EMassFragmentPresence::All);
+		Requirements.AddRequirement<FTestFragment_Float>(EMassFragmentAccess::None, EMassFragmentPresence::None);
+		Requirements.CheckValidity(); // force re-caching of "HasXRequirements" cached properties
+
+		// IntsArchetype has Int but not Float - should match
+		AITEST_TRUE("IntsArchetype should match (has Int, no Float)",
+			Requirements.DoesArchetypeMatchRequirements(IntsArchetype));
+
+		// FloatsIntsArchetype has both Int and Float - should not match (has Float which is excluded)
+		AITEST_FALSE("FloatsIntsArchetype should not match (has excluded Float)",
+			Requirements.DoesArchetypeMatchRequirements(FloatsIntsArchetype));
+
+		return true;
+	}
+};
+IMPLEMENT_AI_INSTANT_TEST(FRequirements_DoesArchetypeMatchRequirements_NoneFragments, "System.Mass.Stress.Requirements.Match.NoneFragments");
+
+
+struct FRequirements_DoesArchetypeMatchRequirements_Tags : FEntityTestBase
+{
+	virtual bool InstantTest() override
+	{
+		// Create archetype with tag
+		FMassArchetypeCompositionDescriptor CompositionWithTag;
+		CompositionWithTag.Add<FTestFragment_Int>();
+		CompositionWithTag.Add<FTestTag_A>();
+
+		FMassArchetypeHandle ArchetypeWithTag = EntityManager->CreateArchetype(CompositionWithTag, {});
+
+		FMassFragmentRequirements Requirements(EntityManager);
+		Requirements.AddTagRequirement<FTestTag_A>(EMassFragmentPresence::All);
+		Requirements.CheckValidity(); // force re-caching of "HasXRequirements" cached properties
+
+		AITEST_TRUE("Archetype with tag should match",
+			Requirements.DoesArchetypeMatchRequirements(ArchetypeWithTag));
+
+		AITEST_FALSE("IntsArchetype (no tag) should not match",
+			Requirements.DoesArchetypeMatchRequirements(IntsArchetype));
+
+		return true;
+	}
+};
+IMPLEMENT_AI_INSTANT_TEST(FRequirements_DoesArchetypeMatchRequirements_Tags, "System.Mass.Stress.Requirements.Match.Tags");
+
+
+//----------------------------------------------------------------------//
+// FMassFragmentRequirements - AddElementRequirement Tests
+//----------------------------------------------------------------------//
+
+struct FRequirements_AddElementRequirement_Fragment : FEntityTestBase
+{
+	virtual bool InstantTest() override
+	{
+		FMassFragmentRequirements Requirements(EntityManager);
+
+		// AddElementRequirement should route fragments to AddRequirement
+		Requirements.AddElementRequirement(FTestFragment_Int::StaticStruct(), EMassFragmentAccess::ReadOnly, EMassFragmentPresence::All);
+		Requirements.CheckValidity(); // force re-caching of "HasXRequirements" cached properties
+
+		AITEST_TRUE("Fragment should be added to RequiredAllFragments",
+			Requirements.GetRequiredAllFragments().Contains<FTestFragment_Int>());
+
+		return true;
+	}
+};
+IMPLEMENT_AI_INSTANT_TEST(FRequirements_AddElementRequirement_Fragment, "System.Mass.Stress.Requirements.Element.AddElementRequirement.Fragment");
+
+
+struct FRequirements_AddElementRequirement_Tag : FEntityTestBase
+{
+	virtual bool InstantTest() override
+	{
+		FMassFragmentRequirements Requirements(EntityManager);
+
+		// AddElementRequirement should route tags to AddTagRequirement
+		Requirements.AddElementRequirement(FTestTag_A::StaticStruct(), EMassFragmentAccess::None, EMassFragmentPresence::All);
+		Requirements.CheckValidity(); // force re-caching of "HasXRequirements" cached properties
+
+		AITEST_TRUE("Tag should be added to RequiredAllTags",
+			Requirements.GetRequiredAllTags().Contains<FTestTag_A>());
+
+		return true;
+	}
+};
+IMPLEMENT_AI_INSTANT_TEST(FRequirements_AddElementRequirement_Tag, "System.Mass.Stress.Requirements.Element.AddElementRequirement.Tag");
+
+
+//----------------------------------------------------------------------//
+// FMassFragmentRequirements - Chunk Fragment Requirements Tests
+//----------------------------------------------------------------------//
+
+struct FRequirements_AddChunkRequirement_All : FEntityTestBase
+{
+	virtual bool InstantTest() override
+	{
+		FMassFragmentRequirements Requirements(EntityManager);
+
+		Requirements.AddChunkRequirement<FTestChunkFragment_Int>(EMassFragmentAccess::ReadOnly, EMassFragmentPresence::All);
+		Requirements.CheckValidity(); // force re-caching of "HasXRequirements" cached properties
+
+		AITEST_TRUE("RequiredAllChunkFragments should contain the chunk fragment",
+			Requirements.GetRequiredAllChunkFragments().Contains<FTestChunkFragment_Int>());
+
+		TConstArrayView<FMassFragmentRequirementDescription> ChunkRequirements = Requirements.GetChunkFragmentRequirements();
+		AITEST_EQUAL("Should have one chunk fragment requirement", ChunkRequirements.Num(), 1);
+
+		return true;
+	}
+};
+IMPLEMENT_AI_INSTANT_TEST(FRequirements_AddChunkRequirement_All, "System.Mass.Stress.Requirements.ChunkFragment.AddChunkRequirement.All");
+
+
+struct FRequirements_AddChunkRequirement_Optional : FEntityTestBase
+{
+	virtual bool InstantTest() override
+	{
+		FMassFragmentRequirements Requirements(EntityManager);
+
+		Requirements.AddChunkRequirement<FTestChunkFragment_Float>(EMassFragmentAccess::ReadWrite, EMassFragmentPresence::Optional);
+		Requirements.CheckValidity(); // force re-caching of "HasXRequirements" cached properties
+
+		AITEST_TRUE("RequiredOptionalChunkFragments should contain the chunk fragment",
+			Requirements.GetRequiredOptionalChunkFragments().Contains<FTestChunkFragment_Float>());
+
+		return true;
+	}
+};
+IMPLEMENT_AI_INSTANT_TEST(FRequirements_AddChunkRequirement_Optional, "System.Mass.Stress.Requirements.ChunkFragment.AddChunkRequirement.Optional");
+
+
+struct FRequirements_AddChunkRequirement_None : FEntityTestBase
+{
+	virtual bool InstantTest() override
+	{
+		FMassFragmentRequirements Requirements(EntityManager);
+
+		Requirements.AddChunkRequirement<FTestChunkFragment_Int>(EMassFragmentAccess::None, EMassFragmentPresence::None);
+		Requirements.CheckValidity(); // force re-caching of "HasXRequirements" cached properties
+
+		AITEST_TRUE("RequiredNoneChunkFragments should contain the chunk fragment",
+			Requirements.GetRequiredNoneChunkFragments().Contains<FTestChunkFragment_Int>());
+
+		// None presence should not add to ChunkFragmentRequirements array
+		AITEST_EQUAL("ChunkFragmentRequirements should be empty for None presence",
+			Requirements.GetChunkFragmentRequirements().Num(), 0);
+
+		return true;
+	}
+};
+IMPLEMENT_AI_INSTANT_TEST(FRequirements_AddChunkRequirement_None, "System.Mass.Stress.Requirements.ChunkFragment.AddChunkRequirement.None");
+
+
+//----------------------------------------------------------------------//
+// FMassFragmentRequirements - Shared Fragment Requirements Tests
+//----------------------------------------------------------------------//
+
+struct FRequirements_AddSharedRequirement_All : FEntityTestBase
+{
+	virtual bool InstantTest() override
+	{
+		FMassFragmentRequirements Requirements(EntityManager);
+
+		Requirements.AddSharedRequirement<FTestSharedFragment_Int>(EMassFragmentAccess::ReadOnly, EMassFragmentPresence::All);
+		Requirements.CheckValidity(); // force re-caching of "HasXRequirements" cached properties
+
+		AITEST_TRUE("RequiredAllSharedFragments should contain the shared fragment",
+			Requirements.GetRequiredAllSharedFragments().Contains<FTestSharedFragment_Int>());
+
+		TConstArrayView<FMassFragmentRequirementDescription> SharedRequirements = Requirements.GetSharedFragmentRequirements();
+		AITEST_EQUAL("Should have one shared fragment requirement", SharedRequirements.Num(), 1);
+
+		return true;
+	}
+};
+IMPLEMENT_AI_INSTANT_TEST(FRequirements_AddSharedRequirement_All, "System.Mass.Stress.Requirements.SharedFragment.AddSharedRequirement.All");
+
+
+struct FRequirements_AddSharedRequirement_Optional : FEntityTestBase
+{
+	virtual bool InstantTest() override
+	{
+		FMassFragmentRequirements Requirements(EntityManager);
+
+		Requirements.AddSharedRequirement<FTestSharedFragment_Float>(EMassFragmentAccess::ReadWrite, EMassFragmentPresence::Optional);
+		Requirements.CheckValidity(); // force re-caching of "HasXRequirements" cached properties
+
+		AITEST_TRUE("RequiredOptionalSharedFragments should contain the shared fragment",
+			Requirements.GetRequiredOptionalSharedFragments().Contains<FTestSharedFragment_Float>());
+
+		return true;
+	}
+};
+IMPLEMENT_AI_INSTANT_TEST(FRequirements_AddSharedRequirement_Optional, "System.Mass.Stress.Requirements.SharedFragment.AddSharedRequirement.Optional");
+
+
+struct FRequirements_AddSharedRequirement_None : FEntityTestBase
+{
+	virtual bool InstantTest() override
+	{
+		FMassFragmentRequirements Requirements(EntityManager);
+
+		Requirements.AddSharedRequirement<FTestSharedFragment_Int>(EMassFragmentAccess::None, EMassFragmentPresence::None);
+		Requirements.CheckValidity(); // force re-caching of "HasXRequirements" cached properties
+
+		AITEST_TRUE("RequiredNoneSharedFragments should contain the shared fragment",
+			Requirements.GetRequiredNoneSharedFragments().Contains<FTestSharedFragment_Int>());
+
+		return true;
+	}
+};
+IMPLEMENT_AI_INSTANT_TEST(FRequirements_AddSharedRequirement_None, "System.Mass.Stress.Requirements.SharedFragment.AddSharedRequirement.None");
+
+
+struct FRequirements_AddConstSharedRequirement_All : FEntityTestBase
+{
+	virtual bool InstantTest() override
+	{
+		FMassFragmentRequirements Requirements(EntityManager);
+
+		Requirements.AddConstSharedRequirement<FTestConstSharedFragment_Int>(EMassFragmentPresence::All);
+		Requirements.CheckValidity(); // force re-caching of "HasXRequirements" cached properties
+
+		AITEST_TRUE("RequiredAllConstSharedFragments should contain the const shared fragment",
+			Requirements.GetRequiredAllConstSharedFragments().Contains<FTestConstSharedFragment_Int>());
+
+		TConstArrayView<FMassFragmentRequirementDescription> ConstSharedRequirements = Requirements.GetConstSharedFragmentRequirements();
+		AITEST_EQUAL("Should have one const shared fragment requirement", ConstSharedRequirements.Num(), 1);
+		// Const shared fragments are always ReadOnly
+		AITEST_EQUAL("Access mode should be ReadOnly", ConstSharedRequirements[0].AccessMode, EMassFragmentAccess::ReadOnly);
+
+		return true;
+	}
+};
+IMPLEMENT_AI_INSTANT_TEST(FRequirements_AddConstSharedRequirement_All, "System.Mass.Stress.Requirements.ConstSharedFragment.AddConstSharedRequirement.All");
+
+
+struct FRequirements_AddConstSharedRequirement_Optional : FEntityTestBase
+{
+	virtual bool InstantTest() override
+	{
+		FMassFragmentRequirements Requirements(EntityManager);
+
+		Requirements.AddConstSharedRequirement<FTestConstSharedFragment_Float>(EMassFragmentPresence::Optional);
+		Requirements.CheckValidity(); // force re-caching of "HasXRequirements" cached properties
+
+		AITEST_TRUE("RequiredOptionalConstSharedFragments should contain the const shared fragment",
+			Requirements.GetRequiredOptionalConstSharedFragments().Contains<FTestConstSharedFragment_Float>());
+
+		return true;
+	}
+};
+IMPLEMENT_AI_INSTANT_TEST(FRequirements_AddConstSharedRequirement_Optional, "System.Mass.Stress.Requirements.ConstSharedFragment.AddConstSharedRequirement.Optional");
+
+
+struct FRequirements_AddConstSharedRequirement_None : FEntityTestBase
+{
+	virtual bool InstantTest() override
+	{
+		FMassFragmentRequirements Requirements(EntityManager);
+
+		Requirements.AddConstSharedRequirement<FTestConstSharedFragment_Int>(EMassFragmentPresence::None);
+		Requirements.CheckValidity(); // force re-caching of "HasXRequirements" cached properties
+
+		AITEST_TRUE("RequiredNoneConstSharedFragments should contain the const shared fragment",
+			Requirements.GetRequiredNoneConstSharedFragments().Contains<FTestConstSharedFragment_Int>());
+
+		return true;
+	}
+};
+IMPLEMENT_AI_INSTANT_TEST(FRequirements_AddConstSharedRequirement_None, "System.Mass.Stress.Requirements.ConstSharedFragment.AddConstSharedRequirement.None");
+
+
+//----------------------------------------------------------------------//
+// FMassFragmentRequirements - Comprehensive Mixed Tests
+//----------------------------------------------------------------------//
+
+struct FRequirements_MixedRequirements : FEntityTestBase
+{
+	virtual bool InstantTest() override
+	{
+		FMassFragmentRequirements Requirements(EntityManager);
+
+		// Add various requirement types - use unified API for sparse element
+		Requirements.AddRequirement<FTestFragment_Int>(EMassFragmentAccess::ReadOnly, EMassFragmentPresence::All);
+		Requirements.AddRequirement<FTestFragment_Float>(EMassFragmentAccess::ReadWrite, EMassFragmentPresence::Optional);
+		Requirements.AddTagRequirement<FTestTag_A>(EMassFragmentPresence::All);
+		Requirements.AddTagRequirement<FTestTag_B>(EMassFragmentPresence::None);
+		Requirements.AddRequirement<FTestFragment_SparseInt>(EMassFragmentAccess::ReadOnly, EMassFragmentPresence::All);
+		Requirements.AddChunkRequirement<FTestChunkFragment_Int>(EMassFragmentAccess::ReadOnly, EMassFragmentPresence::All);
+		Requirements.AddSharedRequirement<FTestSharedFragment_Int>(EMassFragmentAccess::ReadOnly, EMassFragmentPresence::All);
+		Requirements.AddConstSharedRequirement<FTestConstSharedFragment_Int>(EMassFragmentPresence::All);
+		Requirements.CheckValidity(); // force re-caching of "HasXRequirements" cached properties
+
+		// Verify all requirements are tracked
+		AITEST_FALSE("Requirements should not be empty", Requirements.IsEmpty());
+		AITEST_TRUE("Should have positive requirements", Requirements.HasPositiveRequirements());
+		AITEST_TRUE("Should have negative requirements (None tag)", Requirements.HasNegativeRequirements());
+		AITEST_TRUE("Should have optional requirements", Requirements.HasOptionalRequirements());
+		AITEST_TRUE("Should have sparse requirements", Requirements.HasSparseRequirements());
+
+		// Verify specific requirements
+		AITEST_TRUE("Fragment Int in All", Requirements.GetRequiredAllFragments().Contains<FTestFragment_Int>());
+		AITEST_TRUE("Fragment Float in Optional", Requirements.GetRequiredOptionalFragments().Contains<FTestFragment_Float>());
+		AITEST_TRUE("Tag A in All", Requirements.GetRequiredAllTags().Contains<FTestTag_A>());
+		AITEST_TRUE("Tag B in None", Requirements.GetRequiredNoneTags().Contains<FTestTag_B>());
+		AITEST_TRUE("Sparse Int in All", Requirements.GetRequiredAllSparseElements().Contains(FTestFragment_SparseInt::StaticStruct()));
+		AITEST_TRUE("Chunk Int in All", Requirements.GetRequiredAllChunkFragments().Contains<FTestChunkFragment_Int>());
+		AITEST_TRUE("Shared Int in All", Requirements.GetRequiredAllSharedFragments().Contains<FTestSharedFragment_Int>());
+		AITEST_TRUE("ConstShared Int in All", Requirements.GetRequiredAllConstSharedFragments().Contains<FTestConstSharedFragment_Int>());
+
+		return true;
+	}
+};
+IMPLEMENT_AI_INSTANT_TEST(FRequirements_MixedRequirements, "System.Mass.Stress.Requirements.Mixed.Comprehensive");
+
+
+struct FRequirements_Chaining : FEntityTestBase
+{
+	virtual bool InstantTest() override
+	{
+		FMassFragmentRequirements Requirements(EntityManager);
+
+		// Test method chaining - use unified API for sparse element
+		Requirements
+			.AddRequirement<FTestFragment_Int>(EMassFragmentAccess::ReadOnly, EMassFragmentPresence::All)
+			.AddRequirement<FTestFragment_Float>(EMassFragmentAccess::ReadWrite, EMassFragmentPresence::Optional)
+			.AddTagRequirement<FTestTag_A>(EMassFragmentPresence::All)
+			.AddRequirement<FTestFragment_SparseInt>(EMassFragmentAccess::ReadOnly, EMassFragmentPresence::All);
+		Requirements.CheckValidity(); // force re-caching of "HasXRequirements" cached properties
+
+		AITEST_TRUE("Fragment Int added via chaining", Requirements.GetRequiredAllFragments().Contains<FTestFragment_Int>());
+		AITEST_TRUE("Fragment Float added via chaining", Requirements.GetRequiredOptionalFragments().Contains<FTestFragment_Float>());
+		AITEST_TRUE("Tag A added via chaining", Requirements.GetRequiredAllTags().Contains<FTestTag_A>());
+		AITEST_TRUE("Sparse Int added via chaining", Requirements.GetRequiredAllSparseElements().Contains(FTestFragment_SparseInt::StaticStruct()));
+
+		return true;
+	}
+};
+IMPLEMENT_AI_INSTANT_TEST(FRequirements_Chaining, "System.Mass.Stress.Requirements.Mixed.Chaining");
+
+#endif // WITH_MASSENTITY_DEBUG
+
+} // namespace UE::Mass::Test::Requirements
+
+UE_ENABLE_OPTIMIZATION_SHIP
+
+#undef LOCTEXT_NAMESPACE
diff --git a/Engine/Source/Developer/MassEntityTestSuite/Private/StressTests/MassStressSharedFragmentsTest.cpp b/Engine/Source/Developer/MassEntityTestSuite/Private/StressTests/MassStressSharedFragmentsTest.cpp
new file mode 100644
--- /dev/null
+++ b/Engine/Source/Developer/MassEntityTestSuite/Private/StressTests/MassStressSharedFragmentsTest.cpp
@@ -0,0 +1,956 @@
+// Copyright Epic Games, Inc. All Rights Reserved.
+
+#include "AITestsCommon.h"
+#include "MassEntityManager.h"
+#include "MassEntityTypes.h"
+#include "MassEntityTestTypes.h"
+#include "MassArchetypeData.h"
+#include "MassArchetypeTypes.h"
+#include "MassEntityView.h"
+#include "Algo/RandomShuffle.h"
+#include "MassObserverNotificationTypes.h"
+
+#define LOCTEXT_NAMESPACE "MassTest"
+
+UE_DISABLE_OPTIMIZATION_SHIP
+
+namespace FMassEntityTest
+{
+
+//----------------------------------------------------------------------//
+// FMassArchetypeData::SetSharedFragmentsData Tests
+//----------------------------------------------------------------------//
+
+struct FSetSharedFragmentsData_Basic : FEntityTestBase
+{
+	virtual bool InstantTest() override
+	{
+		constexpr int32 InitialValue = 100;
+		constexpr int32 NewValue = 200;
+
+		// Create an archetype with shared fragment
+		FMassArchetypeSharedFragmentValues SharedValues;
+		SharedValues.Add(FSharedStruct::Make<FTestSharedFragment_Int>(InitialValue));
+		SharedValues.Sort();
+
+		FMassEntityHandle Entity = EntityManager->CreateEntity(FloatsArchetype, SharedValues);
+
+		// Verify initial shared fragment value
+		const FTestSharedFragment_Int* InitialFragment = EntityManager->GetSharedFragmentDataPtr<FTestSharedFragment_Int>(Entity);
+		AITEST_NOT_NULL("Entity should have shared fragment after creation", InitialFragment);
+		AITEST_EQUAL("Initial shared fragment value should match", InitialFragment->Value, InitialValue);
+
+		// Get archetype data to call SetSharedFragmentsData
+		FMassArchetypeHandle Archetype = EntityManager->GetArchetypeForEntity(Entity);
+		FMassArchetypeData& ArchetypeData = FMassArchetypeHelper::ArchetypeDataFromHandleChecked(Archetype);
+
+		// Prepare override values
+		TArray<FSharedStruct> Overrides;
+		Overrides.Add(FSharedStruct::Make<FTestSharedFragment_Int>(NewValue));
+
+		// Call SetSharedFragmentsData
+		//ArchetypeData.SetSharedFragmentsData<FSharedStruct>(Entity, MakeArrayView(Overrides));
+		ArchetypeData.SetSharedFragmentsData(Entity, MakeArrayView(Overrides));
+
+		// Verify the entity is still valid and in the same archetype
+		AITEST_TRUE("Entity should still be valid after SetSharedFragmentsData", EntityManager->IsEntityValid(Entity));
+		AITEST_EQUAL("Entity should remain in the same archetype", EntityManager->GetArchetypeForEntity(Entity), Archetype);
+
+		// Verify the shared fragment value changed
+		const FTestSharedFragment_Int* UpdatedFragment = EntityManager->GetSharedFragmentDataPtr<FTestSharedFragment_Int>(Entity);
+		AITEST_NOT_NULL("Entity should still have shared fragment after SetSharedFragmentsData", UpdatedFragment);
+		AITEST_EQUAL("Shared fragment value should be updated", UpdatedFragment->Value, NewValue);
+
+		return true;
+	}
+};
+IMPLEMENT_AI_INSTANT_TEST(FSetSharedFragmentsData_Basic, "System.Mass.Stress.SharedFragments.SetSharedFragmentsData.Basic");
+
+
+struct FSetSharedFragmentsData_MultipleEntitiesDifferentChunks : FEntityTestBase
+{
+	virtual bool InstantTest() override
+	{
+		constexpr int32 ValueA = 10;
+		constexpr int32 ValueB = 20;
+		constexpr int32 ValueC = 30;
+
+		// Create multiple entities with different initial shared fragment values
+		FMassArchetypeSharedFragmentValues SharedValuesA;
+		SharedValuesA.Add(FSharedStruct::Make<FTestSharedFragment_Int>(ValueA));
+		SharedValuesA.Sort();
+
+		FMassArchetypeSharedFragmentValues SharedValuesB;
+		SharedValuesB.Add(FSharedStruct::Make<FTestSharedFragment_Int>(ValueB));
+		SharedValuesB.Sort();
+
+		FMassEntityHandle EntityA = EntityManager->CreateEntity(FloatsArchetype, SharedValuesA);
+		FMassEntityHandle EntityB = EntityManager->CreateEntity(FloatsArchetype, SharedValuesB);
+
+		FMassArchetypeHandle Archetype = EntityManager->GetArchetypeForEntity(EntityA);
+		AITEST_EQUAL("Both entities should be in same archetype", EntityManager->GetArchetypeForEntity(EntityB), Archetype);
+
+		FMassArchetypeData& ArchetypeData = FMassArchetypeHelper::ArchetypeDataFromHandleChecked(Archetype);
+
+		// Initially there should be 2 chunks (one for each shared value)
+		AITEST_EQUAL("Should have 2 chunks initially (one per shared value)", ArchetypeData.GetChunkCount(), 2);
+
+		// Move EntityA to a new shared value (ValueC)
+		TArray<FSharedStruct> OverridesForA;
+		OverridesForA.Add(FSharedStruct::Make<FTestSharedFragment_Int>(ValueC));
+		ArchetypeData.SetSharedFragmentsData<FSharedStruct>(EntityA, OverridesForA);
+
+		// Verify EntityA moved to new chunk with ValueC
+		const FTestSharedFragment_Int* FragmentA = EntityManager->GetSharedFragmentDataPtr<FTestSharedFragment_Int>(EntityA);
+		AITEST_NOT_NULL("EntityA should have shared fragment", FragmentA);
+		AITEST_EQUAL("EntityA's shared fragment should have new value", FragmentA->Value, ValueC);
+
+		// Verify EntityB still has original value
+		const FTestSharedFragment_Int* FragmentB = EntityManager->GetSharedFragmentDataPtr<FTestSharedFragment_Int>(EntityB);
+		AITEST_NOT_NULL("EntityB should have shared fragment", FragmentB);
+		AITEST_EQUAL("EntityB's shared fragment should be unchanged", FragmentB->Value, ValueB);
+
+		// Should now have 3 chunks (ValueA empty but still exists, ValueB, ValueC)
+		// Note: Empty chunks may or may not be cleaned up depending on implementation
+		AITEST_TRUE("Chunk count should be at least 2", ArchetypeData.GetChunkCount() >= 2);
+
+		return true;
+	}
+};
+IMPLEMENT_AI_INSTANT_TEST(FSetSharedFragmentsData_MultipleEntitiesDifferentChunks, "System.Mass.Stress.SharedFragments.SetSharedFragmentsData.MultipleEntitiesDifferentChunks");
+
+
+struct FSetSharedFragmentsData_PreservesFragmentData : FEntityTestBase
+{
+	virtual bool InstantTest() override
+	{
+		constexpr int32 InitialSharedValue = 100;
+		constexpr int32 NewSharedValue = 200;
+		constexpr float TestFloatValue = 42.5f;
+
+		// Create entity with both regular fragment data and shared fragment
+		FMassArchetypeSharedFragmentValues SharedValues;
+		SharedValues.Add(FSharedStruct::Make<FTestSharedFragment_Int>(InitialSharedValue));
+		SharedValues.Sort();
+
+		FMassEntityHandle Entity = EntityManager->CreateEntity(FloatsArchetype, SharedValues);
+
+		// Set a specific value in the regular fragment
+		FTestFragment_Float* FloatFragment = EntityManager->GetFragmentDataPtr<FTestFragment_Float>(Entity);
+		AITEST_NOT_NULL("Entity should have float fragment", FloatFragment);
+		FloatFragment->Value = TestFloatValue;
+
+		// Get archetype and call SetSharedFragmentsData
+		FMassArchetypeHandle Archetype = EntityManager->GetArchetypeForEntity(Entity);
+		FMassArchetypeData& ArchetypeData = FMassArchetypeHelper::ArchetypeDataFromHandleChecked(Archetype);
+
+		TArray<FSharedStruct> Overrides;
+		Overrides.Add(FSharedStruct::Make<FTestSharedFragment_Int>(NewSharedValue));
+		ArchetypeData.SetSharedFragmentsData<FSharedStruct>(Entity, Overrides);
+
+		// Verify regular fragment data is preserved
+		FloatFragment = EntityManager->GetFragmentDataPtr<FTestFragment_Float>(Entity);
+		AITEST_NOT_NULL("Entity should still have float fragment after move", FloatFragment);
+		AITEST_EQUAL("Float fragment value should be preserved after SetSharedFragmentsData", FloatFragment->Value, TestFloatValue);
+
+		// Verify shared fragment was updated
+		const FTestSharedFragment_Int* SharedFragment = EntityManager->GetSharedFragmentDataPtr<FTestSharedFragment_Int>(Entity);
+		AITEST_NOT_NULL("Entity should have shared fragment", SharedFragment);
+		AITEST_EQUAL("Shared fragment should have new value", SharedFragment->Value, NewSharedValue);
+
+		return true;
+	}
+};
+IMPLEMENT_AI_INSTANT_TEST(FSetSharedFragmentsData_PreservesFragmentData, "System.Mass.Stress.SharedFragments.SetSharedFragmentsData.PreservesFragmentData");
+
+
+struct FSetSharedFragmentsData_MultipleSharedFragments : FEntityTestBase
+{
+	virtual bool InstantTest() override
+	{
+		constexpr int32 InitialIntValue = 100;
+		constexpr float InitialFloatValue = 1.5f;
+		constexpr int32 NewIntValue = 200;
+		constexpr float NewFloatValue = 2.5f;
+
+		// Create entity with multiple shared fragments
+		FMassArchetypeSharedFragmentValues SharedValues;
+		SharedValues.Add(FSharedStruct::Make<FTestSharedFragment_Int>(InitialIntValue));
+		SharedValues.Add(FSharedStruct::Make<FTestSharedFragment_Float>(InitialFloatValue));
+		SharedValues.Sort();
+
+		FMassEntityHandle Entity = EntityManager->CreateEntity(FloatsArchetype, SharedValues);
+
+		// Verify initial values
+		const FTestSharedFragment_Int* IntFragment = EntityManager->GetSharedFragmentDataPtr<FTestSharedFragment_Int>(Entity);
+		const FTestSharedFragment_Float* FloatFragment = EntityManager->GetSharedFragmentDataPtr<FTestSharedFragment_Float>(Entity);
+		AITEST_NOT_NULL("Entity should have int shared fragment", IntFragment);
+		AITEST_NOT_NULL("Entity should have float shared fragment", FloatFragment);
+		AITEST_EQUAL("Initial int shared fragment value", IntFragment->Value, InitialIntValue);
+		AITEST_EQUAL("Initial float shared fragment value", FloatFragment->Value, InitialFloatValue);
+
+		// Get archetype and update both shared fragments
+		FMassArchetypeHandle Archetype = EntityManager->GetArchetypeForEntity(Entity);
+		FMassArchetypeData& ArchetypeData = FMassArchetypeHelper::ArchetypeDataFromHandleChecked(Archetype);
+
+		TArray<FSharedStruct> Overrides;
+		Overrides.Add(FSharedStruct::Make<FTestSharedFragment_Int>(NewIntValue));
+		Overrides.Add(FSharedStruct::Make<FTestSharedFragment_Float>(NewFloatValue));
+		ArchetypeData.SetSharedFragmentsData<FSharedStruct>(Entity, Overrides);
+
+		// Verify both shared fragments updated
+		IntFragment = EntityManager->GetSharedFragmentDataPtr<FTestSharedFragment_Int>(Entity);
+		FloatFragment = EntityManager->GetSharedFragmentDataPtr<FTestSharedFragment_Float>(Entity);
+		AITEST_NOT_NULL("Entity should still have int shared fragment", IntFragment);
+		AITEST_NOT_NULL("Entity should still have float shared fragment", FloatFragment);
+		AITEST_EQUAL("Int shared fragment should have new value", IntFragment->Value, NewIntValue);
+		AITEST_EQUAL("Float shared fragment should have new value", FloatFragment->Value, NewFloatValue);
+
+		return true;
+	}
+};
+IMPLEMENT_AI_INSTANT_TEST(FSetSharedFragmentsData_MultipleSharedFragments, "System.Mass.Stress.SharedFragments.SetSharedFragmentsData.MultipleSharedFragments");
+
+
+struct FSetSharedFragmentsData_PartialUpdate : FEntityTestBase
+{
+	virtual bool InstantTest() override
+	{
+		constexpr int32 InitialIntValue = 100;
+		constexpr float InitialFloatValue = 1.5f;
+		constexpr int32 NewIntValue = 200;
+
+		// Create entity with multiple shared fragments
+		FMassArchetypeSharedFragmentValues SharedValues;
+		SharedValues.Add(FSharedStruct::Make<FTestSharedFragment_Int>(InitialIntValue));
+		SharedValues.Add(FSharedStruct::Make<FTestSharedFragment_Float>(InitialFloatValue));
+		SharedValues.Sort();
+
+		FMassEntityHandle Entity = EntityManager->CreateEntity(FloatsArchetype, SharedValues);
+
+		// Get archetype and update only the int shared fragment
+		FMassArchetypeHandle Archetype = EntityManager->GetArchetypeForEntity(Entity);
+		FMassArchetypeData& ArchetypeData = FMassArchetypeHelper::ArchetypeDataFromHandleChecked(Archetype);
+
+		TArray<FSharedStruct> Overrides;
+		Overrides.Add(FSharedStruct::Make<FTestSharedFragment_Int>(NewIntValue));
+		// Note: Not adding FTestSharedFragment_Float to overrides
+		ArchetypeData.SetSharedFragmentsData<FSharedStruct>(Entity, Overrides);
+
+		// Verify int shared fragment updated
+		const FTestSharedFragment_Int* IntFragment = EntityManager->GetSharedFragmentDataPtr<FTestSharedFragment_Int>(Entity);
+		AITEST_NOT_NULL("Entity should have int shared fragment", IntFragment);
+		AITEST_EQUAL("Int shared fragment should have new value", IntFragment->Value, NewIntValue);
+
+		// Verify float shared fragment kept its original value
+		const FTestSharedFragment_Float* FloatFragment = EntityManager->GetSharedFragmentDataPtr<FTestSharedFragment_Float>(Entity);
+		AITEST_NOT_NULL("Entity should have float shared fragment", FloatFragment);
+		AITEST_EQUAL("Float shared fragment should keep original value", FloatFragment->Value, InitialFloatValue);
+
+		return true;
+	}
+};
+IMPLEMENT_AI_INSTANT_TEST(FSetSharedFragmentsData_PartialUpdate, "System.Mass.Stress.SharedFragments.SetSharedFragmentsData.PartialUpdate");
+
+
+struct FSetSharedFragmentsData_MoveToExistingChunk : FEntityTestBase
+{
+	virtual bool InstantTest() override
+	{
+		constexpr int32 ValueA = 10;
+		constexpr int32 ValueB = 20;
+
+		// Create two entities with different shared values - this creates two chunks
+		FMassArchetypeSharedFragmentValues SharedValuesA;
+		SharedValuesA.Add(FSharedStruct::Make<FTestSharedFragment_Int>(ValueA));
+		SharedValuesA.Sort();
+
+		FMassArchetypeSharedFragmentValues SharedValuesB;
+		FSharedStruct SharedValueB = FSharedStruct::Make<FTestSharedFragment_Int>(ValueB);
+		SharedValuesB.Add(SharedValueB);
+		SharedValuesB.Sort();
+
+		FMassEntityHandle EntityA = EntityManager->CreateEntity(FloatsArchetype, SharedValuesA);
+		FMassEntityHandle EntityB = EntityManager->CreateEntity(FloatsArchetype, SharedValuesB);
+
+		FMassArchetypeHandle Archetype = EntityManager->GetArchetypeForEntity(EntityA);
+		FMassArchetypeData& ArchetypeData = FMassArchetypeHelper::ArchetypeDataFromHandleChecked(Archetype);
+
+		const int32 InitialChunkCount = ArchetypeData.GetChunkCount();
+		AITEST_EQUAL("Should have 2 chunks initially", InitialChunkCount, 2);
+
+		// Move EntityA to the same shared value as EntityB
+		TArray<FSharedStruct> Overrides;
+		//Overrides.Add(FSharedStruct::Make<FTestSharedFragment_Int>(ValueB));
+		Overrides.Add(SharedValueB);
+		ArchetypeData.SetSharedFragmentsData<FSharedStruct>(EntityA, Overrides);
+
+		// Verify EntityA now has ValueB
+		const FTestSharedFragment_Int* FragmentA = EntityManager->GetSharedFragmentDataPtr<FTestSharedFragment_Int>(EntityA);
+		AITEST_NOT_NULL("EntityA should have shared fragment", FragmentA);
+		AITEST_EQUAL("EntityA should now have ValueB", FragmentA->Value, ValueB);
+
+		// Verify both entities have the same shared fragment instance (they're in the same chunk)
+		const FTestSharedFragment_Int* FragmentB = EntityManager->GetSharedFragmentDataPtr<FTestSharedFragment_Int>(EntityB);
+		AITEST_NOT_NULL("EntityB should have shared fragment", FragmentB);
+		AITEST_EQUAL("EntityA and EntityB should share the same shared fragment instance", FragmentA, FragmentB);
+
+		return true;
+	}
+};
+IMPLEMENT_AI_INSTANT_TEST(FSetSharedFragmentsData_MoveToExistingChunk, "System.Mass.Stress.SharedFragments.SetSharedFragmentsData.MoveToExistingChunk");
+
+
+struct FSetSharedFragmentsData_WithSparseElements : FEntityTestBase
+{
+	virtual bool InstantTest() override
+	{
+		constexpr int32 InitialSharedValue = 100;
+		constexpr int32 NewSharedValue = 200;
+		constexpr int32 SparseIntValue = 42;
+
+		// Create entity with shared fragment
+		FMassArchetypeSharedFragmentValues SharedValues;
+		SharedValues.Add(FSharedStruct::Make<FTestSharedFragment_Int>(InitialSharedValue));
+		SharedValues.Sort();
+
+		FMassEntityHandle Entity = EntityManager->CreateEntity(FloatsArchetype, SharedValues);
+
+		// Add a sparse element to the entity
+		FTestFragment_SparseInt& SparseInt = EntityManager->AddSparseElementToEntity<FTestFragment_SparseInt>(Entity);
+		SparseInt.Value = SparseIntValue;
+
+		// Verify sparse element exists
+		AITEST_TRUE("Entity should have sparse element before SetSharedFragmentsData",
+			EntityManager->DoesEntityHaveElement(Entity, FTestFragment_SparseInt::StaticStruct()));
+
+		FMassEntityView EntityView(*EntityManager, Entity);
+
+		const FTestFragment_SparseInt* SparseFragment = EntityView.GetSparseFragmentDataPtr<FTestFragment_SparseInt>();
+		AITEST_NOT_NULL("Sparse fragment should exist", SparseFragment);
+		AITEST_EQUAL("Sparse fragment initial value", SparseFragment->Value, SparseIntValue);
+
+		// Get archetype and call SetSharedFragmentsData
+		FMassArchetypeHandle Archetype = EntityManager->GetArchetypeForEntity(Entity);
+		FMassArchetypeData& ArchetypeData = FMassArchetypeHelper::ArchetypeDataFromHandleChecked(Archetype);
+
+		TArray<FSharedStruct> Overrides;
+		Overrides.Add(FSharedStruct::Make<FTestSharedFragment_Int>(NewSharedValue));
+		// this operation invalidates the entity view
+		ArchetypeData.SetSharedFragmentsData<FSharedStruct>(Entity, Overrides);
+
+		// Verify sparse element is preserved after the move
+		AITEST_TRUE("Entity should still have sparse element after SetSharedFragmentsData",
+			EntityManager->DoesEntityHaveElement(Entity, FTestFragment_SparseInt::StaticStruct()));
+
+		FMassEntityView EntityView2(*EntityManager, Entity);
+		SparseFragment = EntityView2.GetSparseFragmentDataPtr<FTestFragment_SparseInt>();
+		AITEST_NOT_NULL("Sparse fragment should still exist after move", SparseFragment);
+		AITEST_EQUAL("Sparse fragment value should be preserved", SparseFragment->Value, SparseIntValue);
+
+		// Verify shared fragment was updated
+		const FTestSharedFragment_Int* SharedFragment = EntityManager->GetSharedFragmentDataPtr<FTestSharedFragment_Int>(Entity);
+		AITEST_NOT_NULL("Shared fragment should exist", SharedFragment);
+		AITEST_EQUAL("Shared fragment should have new value", SharedFragment->Value, NewSharedValue);
+
+		return true;
+	}
+};
+IMPLEMENT_AI_INSTANT_TEST(FSetSharedFragmentsData_WithSparseElements, "System.Mass.Stress.SharedFragments.SetSharedFragmentsData.WithSparseElements");
+
+
+struct FSetSharedFragmentsData_ChunkBoundary : FEntityTestBase
+{
+	virtual bool InstantTest() override
+	{
+		constexpr int32 ValueA = 10;
+		constexpr int32 ValueB = 20;
+
+		// Create enough entities to fill more than one chunk with the same shared value
+		FMassArchetypeSharedFragmentValues SharedValuesA;
+		SharedValuesA.Add(FSharedStruct::Make<FTestSharedFragment_Int>(ValueA));
+		SharedValuesA.Sort();
+
+		const int32 EntitiesPerChunk = FMassArchetypeHelper::ArchetypeDataFromHandleChecked(FloatsArchetype).GetNumEntitiesPerChunk();
+		const int32 NumEntitiesToCreate = static_cast<int32>(EntitiesPerChunk * 1.5f); // Create 1.5 chunks worth
+
+		TArray<FMassEntityHandle> Entities;
+		EntityManager->BatchCreateEntities(FloatsArchetype, SharedValuesA, NumEntitiesToCreate, Entities);
+
+		FMassArchetypeHandle Archetype = EntityManager->GetArchetypeForEntity(Entities[0]);
+		FMassArchetypeData& ArchetypeData = FMassArchetypeHelper::ArchetypeDataFromHandleChecked(Archetype);
+
+		// Move an entity from the middle of the batch to a new shared value
+		const int32 EntityToMoveIndex = NumEntitiesToCreate / 2;
+		FMassEntityHandle EntityToMove = Entities[EntityToMoveIndex];
+
+		TArray<FSharedStruct> Overrides;
+		Overrides.Add(FSharedStruct::Make<FTestSharedFragment_Int>(ValueB));
+		ArchetypeData.SetSharedFragmentsData<FSharedStruct>(EntityToMove, Overrides);
+
+		// Verify the moved entity has the new value
+		const FTestSharedFragment_Int* MovedFragment = EntityManager->GetSharedFragmentDataPtr<FTestSharedFragment_Int>(EntityToMove);
+		AITEST_NOT_NULL("Moved entity should have shared fragment", MovedFragment);
+		AITEST_EQUAL("Moved entity should have new shared value", MovedFragment->Value, ValueB);
+
+		// Verify other entities still have original value
+		for (int32 i = 0; i < NumEntitiesToCreate; ++i)
+		{
+			if (i != EntityToMoveIndex)
+			{
+				const FTestSharedFragment_Int* Fragment = EntityManager->GetSharedFragmentDataPtr<FTestSharedFragment_Int>(Entities[i]);
+				AITEST_NOT_NULL("Other entity should have shared fragment", Fragment);
+				AITEST_EQUAL("Other entity should have original shared value", Fragment->Value, ValueA);
+			}
+		}
+
+		// Total entity count should be preserved
+		AITEST_EQUAL("Total entity count should be preserved", ArchetypeData.GetNumEntities(), NumEntitiesToCreate);
+
+		return true;
+	}
+};
+IMPLEMENT_AI_INSTANT_TEST(FSetSharedFragmentsData_ChunkBoundary, "System.Mass.Stress.SharedFragments.SetSharedFragmentsData.ChunkBoundary");
+
+
+//----------------------------------------------------------------------//
+// FMassArchetypeData::SetSharedFragmentsData Tests - Const Shared Fragments
+//----------------------------------------------------------------------//
+
+struct FSetConstSharedFragmentsData_Basic : FEntityTestBase
+{
+	virtual bool InstantTest() override
+	{
+		constexpr int32 InitialValue = 100;
+		constexpr int32 NewValue = 200;
+
+		// Create an archetype with const shared fragment
+		FMassArchetypeSharedFragmentValues SharedValues;
+		SharedValues.Add(FConstSharedStruct::Make<FTestConstSharedFragment_Int>(InitialValue));
+		SharedValues.Sort();
+
+		FMassEntityHandle Entity = EntityManager->CreateEntity(FloatsArchetype, SharedValues);
+
+		// Verify initial const shared fragment value
+		const FTestConstSharedFragment_Int* InitialFragment = EntityManager->GetConstSharedFragmentDataPtr<FTestConstSharedFragment_Int>(Entity);
+		AITEST_NOT_NULL("Entity should have const shared fragment after creation", InitialFragment);
+		AITEST_EQUAL("Initial const shared fragment value should match", InitialFragment->Value, InitialValue);
+
+		// Get archetype data to call SetSharedFragmentsData
+		FMassArchetypeHandle Archetype = EntityManager->GetArchetypeForEntity(Entity);
+		FMassArchetypeData& ArchetypeData = FMassArchetypeHelper::ArchetypeDataFromHandleChecked(Archetype);
+
+		// Prepare override values
+		TArray<FConstSharedStruct> Overrides;
+		Overrides.Add(FConstSharedStruct::Make<FTestConstSharedFragment_Int>(NewValue));
+
+		// Call SetSharedFragmentsData with const shared fragment
+		ArchetypeData.SetSharedFragmentsData<FConstSharedStruct>(Entity, Overrides);
+
+		// Verify the entity is still valid and in the same archetype
+		AITEST_TRUE("Entity should still be valid after SetSharedFragmentsData", EntityManager->IsEntityValid(Entity));
+		AITEST_EQUAL("Entity should remain in the same archetype", EntityManager->GetArchetypeForEntity(Entity), Archetype);
+
+		// Verify the const shared fragment value changed
+		const FTestConstSharedFragment_Int* UpdatedFragment = EntityManager->GetConstSharedFragmentDataPtr<FTestConstSharedFragment_Int>(Entity);
+		AITEST_NOT_NULL("Entity should still have const shared fragment after SetSharedFragmentsData", UpdatedFragment);
+		AITEST_EQUAL("Const shared fragment value should be updated", UpdatedFragment->Value, NewValue);
+
+		return true;
+	}
+};
+IMPLEMENT_AI_INSTANT_TEST(FSetConstSharedFragmentsData_Basic, "System.Mass.Stress.SharedFragments.SetConstSharedFragmentsData.Basic");
+
+
+struct FSetConstSharedFragmentsData_MultipleEntitiesDifferentChunks : FEntityTestBase
+{
+	virtual bool InstantTest() override
+	{
+		constexpr int32 ValueA = 10;
+		constexpr int32 ValueB = 20;
+		constexpr int32 ValueC = 30;
+
+		// Create multiple entities with different initial const shared fragment values
+		FMassArchetypeSharedFragmentValues SharedValuesA;
+		SharedValuesA.Add(FConstSharedStruct::Make<FTestConstSharedFragment_Int>(ValueA));
+		SharedValuesA.Sort();
+
+		FMassArchetypeSharedFragmentValues SharedValuesB;
+		SharedValuesB.Add(FConstSharedStruct::Make<FTestConstSharedFragment_Int>(ValueB));
+		SharedValuesB.Sort();
+
+		FMassEntityHandle EntityA = EntityManager->CreateEntity(FloatsArchetype, SharedValuesA);
+		FMassEntityHandle EntityB = EntityManager->CreateEntity(FloatsArchetype, SharedValuesB);
+
+		FMassArchetypeHandle Archetype = EntityManager->GetArchetypeForEntity(EntityA);
+		AITEST_EQUAL("Both entities should be in same archetype", EntityManager->GetArchetypeForEntity(EntityB), Archetype);
+
+		FMassArchetypeData& ArchetypeData = FMassArchetypeHelper::ArchetypeDataFromHandleChecked(Archetype);
+
+		// Initially there should be 2 chunks (one for each const shared value)
+		AITEST_EQUAL("Should have 2 chunks initially (one per const shared value)", ArchetypeData.GetChunkCount(), 2);
+
+		// Move EntityA to a new const shared value (ValueC)
+		TArray<FConstSharedStruct> OverridesForA;
+		OverridesForA.Add(FConstSharedStruct::Make<FTestConstSharedFragment_Int>(ValueC));
+		ArchetypeData.SetSharedFragmentsData<FConstSharedStruct>(EntityA, OverridesForA);
+
+		// Verify EntityA moved to new chunk with ValueC
+		const FTestConstSharedFragment_Int* FragmentA = EntityManager->GetConstSharedFragmentDataPtr<FTestConstSharedFragment_Int>(EntityA);
+		AITEST_NOT_NULL("EntityA should have const shared fragment", FragmentA);
+		AITEST_EQUAL("EntityA's const shared fragment should have new value", FragmentA->Value, ValueC);
+
+		// Verify EntityB still has original value
+		const FTestConstSharedFragment_Int* FragmentB = EntityManager->GetConstSharedFragmentDataPtr<FTestConstSharedFragment_Int>(EntityB);
+		AITEST_NOT_NULL("EntityB should have const shared fragment", FragmentB);
+		AITEST_EQUAL("EntityB's const shared fragment should be unchanged", FragmentB->Value, ValueB);
+
+		// Should now have at least 2 chunks
+		AITEST_TRUE("Chunk count should be at least 2", ArchetypeData.GetChunkCount() >= 2);
+
+		return true;
+	}
+};
+IMPLEMENT_AI_INSTANT_TEST(FSetConstSharedFragmentsData_MultipleEntitiesDifferentChunks, "System.Mass.Stress.SharedFragments.SetConstSharedFragmentsData.MultipleEntitiesDifferentChunks");
+
+
+struct FSetConstSharedFragmentsData_PreservesFragmentData : FEntityTestBase
+{
+	virtual bool InstantTest() override
+	{
+		constexpr int32 InitialSharedValue = 100;
+		constexpr int32 NewSharedValue = 200;
+		constexpr float TestFloatValue = 42.5f;
+
+		// Create entity with both regular fragment data and const shared fragment
+		FMassArchetypeSharedFragmentValues SharedValues;
+		SharedValues.Add(FConstSharedStruct::Make<FTestConstSharedFragment_Int>(InitialSharedValue));
+		SharedValues.Sort();
+
+		FMassEntityHandle Entity = EntityManager->CreateEntity(FloatsArchetype, SharedValues);
+
+		// Set a specific value in the regular fragment
+		FTestFragment_Float* FloatFragment = EntityManager->GetFragmentDataPtr<FTestFragment_Float>(Entity);
+		AITEST_NOT_NULL("Entity should have float fragment", FloatFragment);
+		FloatFragment->Value = TestFloatValue;
+
+		// Get archetype and call SetSharedFragmentsData
+		FMassArchetypeHandle Archetype = EntityManager->GetArchetypeForEntity(Entity);
+		FMassArchetypeData& ArchetypeData = FMassArchetypeHelper::ArchetypeDataFromHandleChecked(Archetype);
+
+		TArray<FConstSharedStruct> Overrides;
+		Overrides.Add(FConstSharedStruct::Make<FTestConstSharedFragment_Int>(NewSharedValue));
+		ArchetypeData.SetSharedFragmentsData<FConstSharedStruct>(Entity, Overrides);
+
+		// Verify regular fragment data is preserved
+		FloatFragment = EntityManager->GetFragmentDataPtr<FTestFragment_Float>(Entity);
+		AITEST_NOT_NULL("Entity should still have float fragment after move", FloatFragment);
+		AITEST_EQUAL("Float fragment value should be preserved after SetSharedFragmentsData", FloatFragment->Value, TestFloatValue);
+
+		// Verify const shared fragment was updated
+		const FTestConstSharedFragment_Int* SharedFragment = EntityManager->GetConstSharedFragmentDataPtr<FTestConstSharedFragment_Int>(Entity);
+		AITEST_NOT_NULL("Entity should have const shared fragment", SharedFragment);
+		AITEST_EQUAL("Const shared fragment should have new value", SharedFragment->Value, NewSharedValue);
+
+		return true;
+	}
+};
+IMPLEMENT_AI_INSTANT_TEST(FSetConstSharedFragmentsData_PreservesFragmentData, "System.Mass.Stress.SharedFragments.SetConstSharedFragmentsData.PreservesFragmentData");
+
+
+struct FSetConstSharedFragmentsData_MultipleConstSharedFragments : FEntityTestBase
+{
+	virtual bool InstantTest() override
+	{
+		constexpr int32 InitialIntValue = 100;
+		constexpr float InitialFloatValue = 1.5f;
+		constexpr int32 NewIntValue = 200;
+		constexpr float NewFloatValue = 2.5f;
+
+		// Create entity with multiple const shared fragments
+		FMassArchetypeSharedFragmentValues SharedValues;
+		SharedValues.Add(FConstSharedStruct::Make<FTestConstSharedFragment_Int>(InitialIntValue));
+		SharedValues.Add(FConstSharedStruct::Make<FTestConstSharedFragment_Float>(InitialFloatValue));
+		SharedValues.Sort();
+
+		FMassEntityHandle Entity = EntityManager->CreateEntity(FloatsArchetype, SharedValues);
+
+		// Verify initial values
+		const FTestConstSharedFragment_Int* IntFragment = EntityManager->GetConstSharedFragmentDataPtr<FTestConstSharedFragment_Int>(Entity);
+		const FTestConstSharedFragment_Float* FloatFragment = EntityManager->GetConstSharedFragmentDataPtr<FTestConstSharedFragment_Float>(Entity);
+		AITEST_NOT_NULL("Entity should have int const shared fragment", IntFragment);
+		AITEST_NOT_NULL("Entity should have float const shared fragment", FloatFragment);
+		AITEST_EQUAL("Initial int const shared fragment value", IntFragment->Value, InitialIntValue);
+		AITEST_EQUAL("Initial float const shared fragment value", FloatFragment->Value, InitialFloatValue);
+
+		// Get archetype and update both const shared fragments
+		FMassArchetypeHandle Archetype = EntityManager->GetArchetypeForEntity(Entity);
+		FMassArchetypeData& ArchetypeData = FMassArchetypeHelper::ArchetypeDataFromHandleChecked(Archetype);
+
+		TArray<FConstSharedStruct> Overrides;
+		Overrides.Add(FConstSharedStruct::Make<FTestConstSharedFragment_Int>(NewIntValue));
+		Overrides.Add(FConstSharedStruct::Make<FTestConstSharedFragment_Float>(NewFloatValue));
+		ArchetypeData.SetSharedFragmentsData<FConstSharedStruct>(Entity, Overrides);
+
+		// Verify both const shared fragments updated
+		IntFragment = EntityManager->GetConstSharedFragmentDataPtr<FTestConstSharedFragment_Int>(Entity);
+		FloatFragment = EntityManager->GetConstSharedFragmentDataPtr<FTestConstSharedFragment_Float>(Entity);
+		AITEST_NOT_NULL("Entity should still have int const shared fragment", IntFragment);
+		AITEST_NOT_NULL("Entity should still have float const shared fragment", FloatFragment);
+		AITEST_EQUAL("Int const shared fragment should have new value", IntFragment->Value, NewIntValue);
+		AITEST_EQUAL("Float const shared fragment should have new value", FloatFragment->Value, NewFloatValue);
+
+		return true;
+	}
+};
+IMPLEMENT_AI_INSTANT_TEST(FSetConstSharedFragmentsData_MultipleConstSharedFragments, "System.Mass.Stress.SharedFragments.SetConstSharedFragmentsData.MultipleConstSharedFragments");
+
+
+struct FSetConstSharedFragmentsData_PartialUpdate : FEntityTestBase
+{
+	virtual bool InstantTest() override
+	{
+		constexpr int32 InitialIntValue = 100;
+		constexpr float InitialFloatValue = 1.5f;
+		constexpr int32 NewIntValue = 200;
+
+		// Create entity with multiple const shared fragments
+		FMassArchetypeSharedFragmentValues SharedValues;
+		SharedValues.Add(FConstSharedStruct::Make<FTestConstSharedFragment_Int>(InitialIntValue));
+		SharedValues.Add(FConstSharedStruct::Make<FTestConstSharedFragment_Float>(InitialFloatValue));
+		SharedValues.Sort();
+
+		FMassEntityHandle Entity = EntityManager->CreateEntity(FloatsArchetype, SharedValues);
+
+		// Get archetype and update only the int const shared fragment
+		FMassArchetypeHandle Archetype = EntityManager->GetArchetypeForEntity(Entity);
+		FMassArchetypeData& ArchetypeData = FMassArchetypeHelper::ArchetypeDataFromHandleChecked(Archetype);
+
+		TArray<FConstSharedStruct> Overrides;
+		Overrides.Add(FConstSharedStruct::Make<FTestConstSharedFragment_Int>(NewIntValue));
+		// Note: Not adding FTestConstSharedFragment_Float to overrides
+		ArchetypeData.SetSharedFragmentsData<FConstSharedStruct>(Entity, Overrides);
+
+		// Verify int const shared fragment updated
+		const FTestConstSharedFragment_Int* IntFragment = EntityManager->GetConstSharedFragmentDataPtr<FTestConstSharedFragment_Int>(Entity);
+		AITEST_NOT_NULL("Entity should have int const shared fragment", IntFragment);
+		AITEST_EQUAL("Int const shared fragment should have new value", IntFragment->Value, NewIntValue);
+
+		// Verify float const shared fragment kept its original value
+		const FTestConstSharedFragment_Float* FloatFragment = EntityManager->GetConstSharedFragmentDataPtr<FTestConstSharedFragment_Float>(Entity);
+		AITEST_NOT_NULL("Entity should have float const shared fragment", FloatFragment);
+		AITEST_EQUAL("Float const shared fragment should keep original value", FloatFragment->Value, InitialFloatValue);
+
+		return true;
+	}
+};
+IMPLEMENT_AI_INSTANT_TEST(FSetConstSharedFragmentsData_PartialUpdate, "System.Mass.Stress.SharedFragments.SetConstSharedFragmentsData.PartialUpdate");
+
+
+struct FSetConstSharedFragmentsData_MoveToExistingChunk : FEntityTestBase
+{
+	virtual bool InstantTest() override
+	{
+		constexpr int32 ValueA = 10;
+		constexpr int32 ValueB = 20;
+
+		// Create two entities with different const shared values - this creates two chunks
+		FMassArchetypeSharedFragmentValues SharedValuesA;
+		SharedValuesA.Add(FConstSharedStruct::Make<FTestConstSharedFragment_Int>(ValueA));
+		SharedValuesA.Sort();
+
+		FMassArchetypeSharedFragmentValues SharedValuesB;
+		FConstSharedStruct SharedValueB = FConstSharedStruct::Make<FTestConstSharedFragment_Int>(ValueB);
+		SharedValuesB.Add(SharedValueB);
+		SharedValuesB.Sort();
+
+		FMassEntityHandle EntityA = EntityManager->CreateEntity(FloatsArchetype, SharedValuesA);
+		FMassEntityHandle EntityB = EntityManager->CreateEntity(FloatsArchetype, SharedValuesB);
+
+		FMassArchetypeHandle Archetype = EntityManager->GetArchetypeForEntity(EntityA);
+		FMassArchetypeData& ArchetypeData = FMassArchetypeHelper::ArchetypeDataFromHandleChecked(Archetype);
+
+		const int32 InitialChunkCount = ArchetypeData.GetChunkCount();
+		AITEST_EQUAL("Should have 2 chunks initially", InitialChunkCount, 2);
+
+		// Move EntityA to the same const shared value as EntityB
+		TArray<FConstSharedStruct> Overrides;
+		Overrides.Add(SharedValueB);
+		ArchetypeData.SetSharedFragmentsData<FConstSharedStruct>(EntityA, Overrides);
+
+		// Verify EntityA now has ValueB
+		const FTestConstSharedFragment_Int* FragmentA = EntityManager->GetConstSharedFragmentDataPtr<FTestConstSharedFragment_Int>(EntityA);
+		AITEST_NOT_NULL("EntityA should have const shared fragment", FragmentA);
+		AITEST_EQUAL("EntityA should now have ValueB", FragmentA->Value, ValueB);
+
+		// Verify both entities have the same const shared fragment instance (they're in the same chunk)
+		const FTestConstSharedFragment_Int* FragmentB = EntityManager->GetConstSharedFragmentDataPtr<FTestConstSharedFragment_Int>(EntityB);
+		AITEST_NOT_NULL("EntityB should have const shared fragment", FragmentB);
+		AITEST_EQUAL("EntityA and EntityB should share the same const shared fragment instance", FragmentA, FragmentB);
+
+		return true;
+	}
+};
+IMPLEMENT_AI_INSTANT_TEST(FSetConstSharedFragmentsData_MoveToExistingChunk, "System.Mass.Stress.SharedFragments.SetConstSharedFragmentsData.MoveToExistingChunk");
+
+
+struct FSetConstSharedFragmentsData_WithSparseElements : FEntityTestBase
+{
+	virtual bool InstantTest() override
+	{
+		constexpr int32 InitialSharedValue = 100;
+		constexpr int32 NewSharedValue = 200;
+		constexpr int32 SparseIntValue = 42;
+
+		// Create entity with const shared fragment
+		FMassArchetypeSharedFragmentValues SharedValues;
+		SharedValues.Add(FConstSharedStruct::Make<FTestConstSharedFragment_Int>(InitialSharedValue));
+		SharedValues.Sort();
+
+		FMassEntityHandle Entity = EntityManager->CreateEntity(FloatsArchetype, SharedValues);
+
+		// Add a sparse element to the entity
+		FTestFragment_SparseInt& SparseInt = EntityManager->AddSparseElementToEntity<FTestFragment_SparseInt>(Entity);
+		SparseInt.Value = SparseIntValue;
+
+		// Verify sparse element exists
+		AITEST_TRUE("Entity should have sparse element before SetSharedFragmentsData",
+			EntityManager->DoesEntityHaveElement(Entity, FTestFragment_SparseInt::StaticStruct()));
+
+		FMassEntityView EntityView(*EntityManager, Entity);
+
+		const FTestFragment_SparseInt* SparseFragment = EntityView.GetSparseFragmentDataPtr<FTestFragment_SparseInt>();
+		AITEST_NOT_NULL("Sparse fragment should exist", SparseFragment);
+		AITEST_EQUAL("Sparse fragment initial value", SparseFragment->Value, SparseIntValue);
+
+		// Get archetype and call SetSharedFragmentsData
+		FMassArchetypeHandle Archetype = EntityManager->GetArchetypeForEntity(Entity);
+		FMassArchetypeData& ArchetypeData = FMassArchetypeHelper::ArchetypeDataFromHandleChecked(Archetype);
+
+		TArray<FConstSharedStruct> Overrides;
+		Overrides.Add(FConstSharedStruct::Make<FTestConstSharedFragment_Int>(NewSharedValue));
+		// this operation invalidates the entity view
+		ArchetypeData.SetSharedFragmentsData<FConstSharedStruct>(Entity, Overrides);
+
+		// Verify sparse element is preserved after the move
+		AITEST_TRUE("Entity should still have sparse element after SetSharedFragmentsData",
+			EntityManager->DoesEntityHaveElement(Entity, FTestFragment_SparseInt::StaticStruct()));
+
+		FMassEntityView NewEntityView(*EntityManager, Entity);
+		SparseFragment = NewEntityView.GetSparseFragmentDataPtr<FTestFragment_SparseInt>();
+		AITEST_NOT_NULL("Sparse fragment should still exist after move", SparseFragment);
+		AITEST_EQUAL("Sparse fragment value should be preserved", SparseFragment->Value, SparseIntValue);
+
+		// Verify const shared fragment was updated
+		const FTestConstSharedFragment_Int* SharedFragment = EntityManager->GetConstSharedFragmentDataPtr<FTestConstSharedFragment_Int>(Entity);
+		AITEST_NOT_NULL("Const shared fragment should exist", SharedFragment);
+		AITEST_EQUAL("Const shared fragment should have new value", SharedFragment->Value, NewSharedValue);
+
+		return true;
+	}
+};
+IMPLEMENT_AI_INSTANT_TEST(FSetConstSharedFragmentsData_WithSparseElements, "System.Mass.Stress.SharedFragments.SetConstSharedFragmentsData.WithSparseElements");
+
+
+struct FSetConstSharedFragmentsData_ChunkBoundary : FEntityTestBase
+{
+	virtual bool InstantTest() override
+	{
+		constexpr int32 ValueA = 10;
+		constexpr int32 ValueB = 20;
+
+		// Create enough entities to fill more than one chunk with the same const shared value
+		FMassArchetypeSharedFragmentValues SharedValuesA;
+		SharedValuesA.Add(FConstSharedStruct::Make<FTestConstSharedFragment_Int>(ValueA));
+		SharedValuesA.Sort();
+
+		const int32 EntitiesPerChunk = FMassArchetypeHelper::ArchetypeDataFromHandleChecked(FloatsArchetype).GetNumEntitiesPerChunk();
+		const int32 NumEntitiesToCreate = static_cast<int32>(EntitiesPerChunk * 1.5f); // Create 1.5 chunks worth
+
+		TArray<FMassEntityHandle> Entities;
+		EntityManager->BatchCreateEntities(FloatsArchetype, SharedValuesA, NumEntitiesToCreate, Entities);
+
+		FMassArchetypeHandle Archetype = EntityManager->GetArchetypeForEntity(Entities[0]);
+		FMassArchetypeData& ArchetypeData = FMassArchetypeHelper::ArchetypeDataFromHandleChecked(Archetype);
+
+		// Move an entity from the middle of the batch to a new const shared value
+		const int32 EntityToMoveIndex = NumEntitiesToCreate / 2;
+		FMassEntityHandle EntityToMove = Entities[EntityToMoveIndex];
+
+		TArray<FConstSharedStruct> Overrides;
+		Overrides.Add(FConstSharedStruct::Make<FTestConstSharedFragment_Int>(ValueB));
+		ArchetypeData.SetSharedFragmentsData<FConstSharedStruct>(EntityToMove, Overrides);
+
+		// Verify the moved entity has the new value
+		const FTestConstSharedFragment_Int* MovedFragment = EntityManager->GetConstSharedFragmentDataPtr<FTestConstSharedFragment_Int>(EntityToMove);
+		AITEST_NOT_NULL("Moved entity should have const shared fragment", MovedFragment);
+		AITEST_EQUAL("Moved entity should have new const shared value", MovedFragment->Value, ValueB);
+
+		// Verify other entities still have original value
+		for (int32 i = 0; i < NumEntitiesToCreate; ++i)
+		{
+			if (i != EntityToMoveIndex)
+			{
+				const FTestConstSharedFragment_Int* Fragment = EntityManager->GetConstSharedFragmentDataPtr<FTestConstSharedFragment_Int>(Entities[i]);
+				AITEST_NOT_NULL("Other entity should have const shared fragment", Fragment);
+				AITEST_EQUAL("Other entity should have original const shared value", Fragment->Value, ValueA);
+			}
+		}
+
+		// Total entity count should be preserved
+		AITEST_EQUAL("Total entity count should be preserved", ArchetypeData.GetNumEntities(), NumEntitiesToCreate);
+
+		return true;
+	}
+};
+IMPLEMENT_AI_INSTANT_TEST(FSetConstSharedFragmentsData_ChunkBoundary, "System.Mass.Stress.SharedFragments.SetConstSharedFragmentsData.ChunkBoundary");
+
+
+//----------------------------------------------------------------------//
+// FMassArchetypeData::SetSharedFragmentsData Tests - Mixed Shared/Const
+//----------------------------------------------------------------------//
+
+struct FSetSharedFragmentsData_MixedSharedAndConst : FEntityTestBase
+{
+	virtual bool InstantTest() override
+	{
+		constexpr int32 InitialSharedIntValue = 100;
+		constexpr int32 InitialConstIntValue = 200;
+		constexpr int32 NewSharedIntValue = 150;
+		constexpr int32 NewConstIntValue = 250;
+
+		// Create entity with both shared and const shared fragments
+		FMassArchetypeSharedFragmentValues SharedValues;
+		SharedValues.Add(FSharedStruct::Make<FTestSharedFragment_Int>(InitialSharedIntValue));
+		SharedValues.Add(FConstSharedStruct::Make<FTestConstSharedFragment_Int>(InitialConstIntValue));
+		SharedValues.Sort();
+
+		FMassEntityHandle Entity = EntityManager->CreateEntity(FloatsArchetype, SharedValues);
+
+		// Verify initial values
+		const FTestSharedFragment_Int* SharedFragment = EntityManager->GetSharedFragmentDataPtr<FTestSharedFragment_Int>(Entity);
+		const FTestConstSharedFragment_Int* ConstFragment = EntityManager->GetConstSharedFragmentDataPtr<FTestConstSharedFragment_Int>(Entity);
+		AITEST_NOT_NULL("Entity should have shared fragment", SharedFragment);
+		AITEST_NOT_NULL("Entity should have const shared fragment", ConstFragment);
+		AITEST_EQUAL("Initial shared fragment value", SharedFragment->Value, InitialSharedIntValue);
+		AITEST_EQUAL("Initial const shared fragment value", ConstFragment->Value, InitialConstIntValue);
+
+		FMassArchetypeHandle Archetype = EntityManager->GetArchetypeForEntity(Entity);
+		FMassArchetypeData& ArchetypeData = FMassArchetypeHelper::ArchetypeDataFromHandleChecked(Archetype);
+
+		// Update only the mutable shared fragment
+		{
+			TArray<FSharedStruct> SharedOverrides;
+			SharedOverrides.Add(FSharedStruct::Make<FTestSharedFragment_Int>(NewSharedIntValue));
+			ArchetypeData.SetSharedFragmentsData<FSharedStruct>(Entity, SharedOverrides);
+		}
+
+		// Verify shared fragment updated, const unchanged
+		SharedFragment = EntityManager->GetSharedFragmentDataPtr<FTestSharedFragment_Int>(Entity);
+		ConstFragment = EntityManager->GetConstSharedFragmentDataPtr<FTestConstSharedFragment_Int>(Entity);
+		AITEST_NOT_NULL("Entity should still have shared fragment", SharedFragment);
+		AITEST_NOT_NULL("Entity should still have const shared fragment", ConstFragment);
+		AITEST_EQUAL("Shared fragment should have new value", SharedFragment->Value, NewSharedIntValue);
+		AITEST_EQUAL("Const shared fragment should be unchanged", ConstFragment->Value, InitialConstIntValue);
+
+		// Now update only the const shared fragment
+		{
+			TArray<FConstSharedStruct> ConstOverrides;
+			ConstOverrides.Add(FConstSharedStruct::Make<FTestConstSharedFragment_Int>(NewConstIntValue));
+			ArchetypeData.SetSharedFragmentsData<FConstSharedStruct>(Entity, ConstOverrides);
+		}
+
+		// Verify both are now updated
+		SharedFragment = EntityManager->GetSharedFragmentDataPtr<FTestSharedFragment_Int>(Entity);
+		ConstFragment = EntityManager->GetConstSharedFragmentDataPtr<FTestConstSharedFragment_Int>(Entity);
+		AITEST_NOT_NULL("Entity should still have shared fragment", SharedFragment);
+		AITEST_NOT_NULL("Entity should still have const shared fragment", ConstFragment);
+		AITEST_EQUAL("Shared fragment should retain its value", SharedFragment->Value, NewSharedIntValue);
+		AITEST_EQUAL("Const shared fragment should have new value", ConstFragment->Value, NewConstIntValue);
+
+		return true;
+	}
+};
+IMPLEMENT_AI_INSTANT_TEST(FSetSharedFragmentsData_MixedSharedAndConst, "System.Mass.Stress.SharedFragments.SetSharedFragmentsData.MixedSharedAndConst");
+
+
+struct FSetSharedFragmentsData_MixedPreservesOtherType : FEntityTestBase
+{
+	virtual bool InstantTest() override
+	{
+		constexpr int32 SharedIntValue = 100;
+		constexpr float SharedFloatValue = 1.5f;
+		constexpr int32 ConstIntValue = 200;
+		constexpr float ConstFloatValue = 2.5f;
+		constexpr int32 NewSharedIntValue = 150;
+		constexpr int32 NewConstIntValue = 250;
+
+		// Create entity with multiple shared and const shared fragments
+		FMassArchetypeSharedFragmentValues SharedValues;
+		SharedValues.Add(FSharedStruct::Make<FTestSharedFragment_Int>(SharedIntValue));
+		SharedValues.Add(FSharedStruct::Make<FTestSharedFragment_Float>(SharedFloatValue));
+		SharedValues.Add(FConstSharedStruct::Make<FTestConstSharedFragment_Int>(ConstIntValue));
+		SharedValues.Add(FConstSharedStruct::Make<FTestConstSharedFragment_Float>(ConstFloatValue));
+		SharedValues.Sort();
+
+		FMassEntityHandle Entity = EntityManager->CreateEntity(FloatsArchetype, SharedValues);
+
+		FMassArchetypeHandle Archetype = EntityManager->GetArchetypeForEntity(Entity);
+		FMassArchetypeData& ArchetypeData = FMassArchetypeHelper::ArchetypeDataFromHandleChecked(Archetype);
+
+		// Update only one mutable shared fragment (Int), keeping Float unchanged
+		{
+			TArray<FSharedStruct> SharedOverrides;
+			SharedOverrides.Add(FSharedStruct::Make<FTestSharedFragment_Int>(NewSharedIntValue));
+			ArchetypeData.SetSharedFragmentsData<FSharedStruct>(Entity, SharedOverrides);
+		}
+
+		// Verify only the Int shared fragment changed
+		const FTestSharedFragment_Int* SharedInt = EntityManager->GetSharedFragmentDataPtr<FTestSharedFragment_Int>(Entity);
+		const FTestSharedFragment_Float* SharedFloat = EntityManager->GetSharedFragmentDataPtr<FTestSharedFragment_Float>(Entity);
+		const FTestConstSharedFragment_Int* ConstInt = EntityManager->GetConstSharedFragmentDataPtr<FTestConstSharedFragment_Int>(Entity);
+		const FTestConstSharedFragment_Float* ConstFloat = EntityManager->GetConstSharedFragmentDataPtr<FTestConstSharedFragment_Float>(Entity);
+
+		AITEST_NOT_NULL("SharedInt should exist", SharedInt);
+		AITEST_NOT_NULL("SharedFloat should exist", SharedFloat);
+		AITEST_NOT_NULL("ConstInt should exist", ConstInt);
+		AITEST_NOT_NULL("ConstFloat should exist", ConstFloat);
+
+		AITEST_EQUAL("SharedInt should have new value", SharedInt->Value, NewSharedIntValue);
+		AITEST_EQUAL("SharedFloat should be unchanged", SharedFloat->Value, SharedFloatValue);
+		AITEST_EQUAL("ConstInt should be unchanged", ConstInt->Value, ConstIntValue);
+		AITEST_EQUAL("ConstFloat should be unchanged", ConstFloat->Value, ConstFloatValue);
+
+		// Update only one const shared fragment (Int), keeping Float unchanged
+		{
+			TArray<FConstSharedStruct> ConstOverrides;
+			ConstOverrides.Add(FConstSharedStruct::Make<FTestConstSharedFragment_Int>(NewConstIntValue));
+			ArchetypeData.SetSharedFragmentsData<FConstSharedStruct>(Entity, ConstOverrides);
+		}
+
+		// Re-fetch and verify
+		SharedInt = EntityManager->GetSharedFragmentDataPtr<FTestSharedFragment_Int>(Entity);
+		SharedFloat = EntityManager->GetSharedFragmentDataPtr<FTestSharedFragment_Float>(Entity);
+		ConstInt = EntityManager->GetConstSharedFragmentDataPtr<FTestConstSharedFragment_Int>(Entity);
+		ConstFloat = EntityManager->GetConstSharedFragmentDataPtr<FTestConstSharedFragment_Float>(Entity);
+
+		AITEST_NOT_NULL("SharedInt should still exist", SharedInt);
+		AITEST_NOT_NULL("SharedFloat should still exist", SharedFloat);
+		AITEST_NOT_NULL("ConstInt should still exist", ConstInt);
+		AITEST_NOT_NULL("ConstFloat should still exist", ConstFloat);
+
+		AITEST_EQUAL("SharedInt should retain value", SharedInt->Value, NewSharedIntValue);
+		AITEST_EQUAL("SharedFloat should still be unchanged", SharedFloat->Value, SharedFloatValue);
+		AITEST_EQUAL("ConstInt should have new value", ConstInt->Value, NewConstIntValue);
+		AITEST_EQUAL("ConstFloat should still be unchanged", ConstFloat->Value, ConstFloatValue);
+
+		return true;
+	}
+};
+IMPLEMENT_AI_INSTANT_TEST(FSetSharedFragmentsData_MixedPreservesOtherType, "System.Mass.Stress.SharedFragments.SetSharedFragmentsData.MixedPreservesOtherType");
+
+} // FMassEntityTest
+
+UE_ENABLE_OPTIMIZATION_SHIP
+
+#undef LOCTEXT_NAMESPACE
diff --git a/Engine/Source/Developer/MassEntityTestSuite/Private/StressTests/MassStressSparseElementsTest.cpp b/Engine/Source/Developer/MassEntityTestSuite/Private/StressTests/MassStressSparseElementsTest.cpp
new file mode 100644
--- /dev/null
+++ b/Engine/Source/Developer/MassEntityTestSuite/Private/StressTests/MassStressSparseElementsTest.cpp
@@ -0,0 +1,2506 @@
+// Copyright Epic Games, Inc. All Rights Reserved.
+
+#include "AITestsCommon.h"
+#include "MassArchetypeData.h"
+#include "MassEntityBuilder.h"
+#include "MassEntityManager.h"
+#include "MassEntityMacros.h"
+#include "MassEntityTestTypes.h"
+#include "MassEntityView.h"
+#include "MassExecutionContext.h"
+#include "MassExecutor.h"
+
+#define LOCTEXT_NAMESPACE "MassTest"
+
+UE_DISABLE_OPTIMIZATION_SHIP
+
+// to be enabled once this functionality gets implemented
+#define WITH_ENTITY_MOVING_AFFECTING_COMPOSITION 0
+
+namespace UE::Mass::Test::SparseElements
+{
+#if WITH_MASSENTITY_DEBUG
+
+	struct FSparseElements_Storage_ConfigureType : FEntityTestBase
+	{
+		virtual bool InstantTest() override
+		{
+			UE::Mass::FSparseElementsStorage Storage;
+
+			// Configure custom chunk size
+			UE::Mass::FSparseElementsStorage::FTypeConfig Config(FTestFragment_SparseInt::StaticStruct());
+			// change the default chunk size, for testing purposes.
+			Config.SetElementsPerChunk(64); 
+
+			Storage.ConfigureType(FTestFragment_SparseInt::StaticStruct(), Config);
+
+			// Create entities beyond one chunk
+			constexpr int32 NumEntities = 200;
+			TArray<FMassEntityHandle> Entities;
+			EntityManager->BatchCreateEntities(IntsArchetype, NumEntities, Entities);
+
+			// Add sparse elements - should use configured chunk size
+			for (const FMassEntityHandle& Entity : Entities)
+			{
+				FStructView View = Storage.AddElementToEntity(Entity, FTestFragment_SparseInt::StaticStruct());
+				AITEST_TRUE(TEXT("Added element with custom chunk size"), View.IsValid());
+			}
+
+			// Verify we can access all of them
+			for (const FMassEntityHandle& Entity : Entities)
+			{
+				FStructView View = Storage.GetMutableElementDataForEntity(Entity, FTestFragment_SparseInt::StaticStruct());
+				AITEST_TRUE(TEXT("Element accessible with custom chunk size"), View.IsValid());
+			}
+
+			return true;
+		}
+	};
+	IMPLEMENT_AI_INSTANT_TEST(FSparseElements_Storage_ConfigureType, "System.Mass.Stress.Sparse.Storage.ConfigureType");
+
+	struct FSparseElements_Storage_OnEntitiesDestroyed : FEntityTestBase
+	{
+		virtual bool InstantTest() override
+		{
+			UE::Mass::FSparseElementsStorage& Storage = EntityManager->DebugGetSparseElementsStorage();
+
+			constexpr int32 NumEntities = 10;
+			TArray<FMassEntityHandle> Entities;
+			EntityManager->BatchCreateEntities(IntsArchetype, NumEntities, Entities);
+
+			// Add sparse elements to all
+			for (const FMassEntityHandle& Entity : Entities)
+			{
+				FStructView View = EntityManager->AddSparseElementToEntity(Entity, FTestFragment_SparseInt::StaticStruct());
+				AITEST_TRUE(TEXT("Element added successfully"), View.IsValid());
+			}
+
+			// Destroy subset - test the bug fix (return -> continue)
+			TArray<FMassEntityHandle> EntitiesToDestroy = {Entities[2], Entities[5], Entities[7]};
+			EntityManager->BatchDestroyEntities(EntitiesToDestroy);
+
+			// Verify destroyed entities no longer have elements
+			for (const FMassEntityHandle& Entity : EntitiesToDestroy)
+			{
+				FStructView View = Storage.GetMutableElementDataForEntity(Entity, FTestFragment_SparseInt::StaticStruct());
+				AITEST_FALSE(TEXT("Destroyed entity has no sparse element"), View.IsValid());
+			}
+
+			// Verify remaining entities still have elements (validates bug fix)
+			for (int32 Index = 0; Index < NumEntities; ++Index)
+			{
+				if (!EntitiesToDestroy.Contains(Entities[Index]))
+				{
+					FStructView View = Storage.GetMutableElementDataForEntity(Entities[Index], FTestFragment_SparseInt::StaticStruct());
+					AITEST_TRUE(TEXT("Remaining entity still has sparse element"), View.IsValid());
+				}
+			}
+
+			return true;
+		}
+	};
+	IMPLEMENT_AI_INSTANT_TEST(FSparseElements_Storage_OnEntitiesDestroyed, "System.Mass.Stress.Sparse.Storage.OnEntitiesDestroyed");
+
+	struct FSparseElements_MultipleTypesPerEntity : FEntityTestBase
+	{
+		virtual bool InstantTest() override
+		{
+			const FMassEntityHandle Entity = EntityManager->CreateEntity(IntsArchetype);
+
+			// Add 3 different sparse types
+			FStructView IntView = EntityManager->AddSparseElementToEntity(Entity, FTestFragment_SparseInt::StaticStruct());
+			FStructView FloatView = EntityManager->AddSparseElementToEntity(Entity, FTestFragment_SparseFloat::StaticStruct());
+			FStructView TagView = EntityManager->AddSparseElementToEntity(Entity, FTestTag_SparseA::StaticStruct());
+
+			AITEST_TRUE(TEXT("Sparse int added"), IntView.IsValid());
+			AITEST_TRUE(TEXT("Sparse float added"), FloatView.IsValid());
+			AITEST_TRUE(TEXT("Added sparse tag view is invalid"), TagView.IsValid() == false);
+
+			// Verify all are present
+			FMassEntityView EntityView(*EntityManager, Entity);
+			AITEST_TRUE(TEXT("Has sparse int"), EntityView.HasElement<FTestFragment_SparseInt>());
+			AITEST_TRUE(TEXT("Has sparse float"), EntityView.HasElement<FTestFragment_SparseFloat>());
+			AITEST_TRUE(TEXT("Has sparse tag"), EntityView.HasElement<FTestTag_SparseA>());
+
+			// Remove one, verify others remain
+			EntityManager->RemoveSparseElementFromEntity(Entity, FTestFragment_SparseFloat::StaticStruct());
+
+			AITEST_TRUE(TEXT("Still has sparse int"), EntityView.HasElement<FTestFragment_SparseInt>());
+			AITEST_FALSE(TEXT("No longer has sparse float"), EntityView.HasElement<FTestFragment_SparseFloat>());
+			AITEST_TRUE(TEXT("Still has sparse tag"), EntityView.HasElement<FTestTag_SparseA>());
+
+			return true;
+		}
+	};
+	IMPLEMENT_AI_INSTANT_TEST(FSparseElements_MultipleTypesPerEntity, "System.Mass.Stress.Sparse.MultipleTypes");
+
+	struct FSparseElements_RemoveNonExistent : FEntityTestBase
+	{
+		virtual bool InstantTest() override
+		{
+			const FMassEntityHandle Entity = EntityManager->CreateEntity(IntsArchetype);
+
+			// Try to remove element that was never added
+			EntityManager->RemoveSparseElementFromEntity(Entity, FTestFragment_SparseInt::StaticStruct());
+
+			// Add, remove, try to remove again
+			FStructView _ = EntityManager->AddSparseElementToEntity(Entity, FTestFragment_SparseInt::StaticStruct());
+			EntityManager->RemoveSparseElementFromEntity(Entity, FTestFragment_SparseInt::StaticStruct());
+
+			return true;
+		}
+	};
+	IMPLEMENT_AI_INSTANT_TEST(FSparseElements_RemoveNonExistent, "System.Mass.Stress.Sparse.RemoveNonExistent");
+
+	struct FSparseElements_DataPersistence : FEntityTestBase
+	{
+		virtual bool InstantTest() override
+		{
+			constexpr int32 NumEntities = 20;
+			TArray<FMassEntityHandle> Entities;
+			EntityManager->BatchCreateEntities(IntsArchetype, NumEntities, Entities);
+
+			// Add sparse fragments with unique values
+			for (int32 Index = 0; Index < NumEntities; ++Index)
+			{
+				FStructView View = EntityManager->AddSparseElementToEntity(Entities[Index], FTestFragment_SparseInt::StaticStruct());
+				FTestFragment_SparseInt* Fragment = View.GetPtr<FTestFragment_SparseInt>();
+				Fragment->Value = Index * 100;
+			}
+
+			// Move entities to different archetype
+			EntityManager->BatchChangeFragmentCompositionForEntities(
+				{FMassArchetypeEntityCollection(IntsArchetype, Entities, FMassArchetypeEntityCollection::EDuplicatesHandling::NoDuplicates)},
+				FMassFragmentBitSet(*FTestFragment_Float::StaticStruct()),
+				{}
+			);
+
+			// Verify data persisted through move
+			for (int32 Index = 0; Index < NumEntities; ++Index)
+			{
+				FMassEntityView EntityView(*EntityManager, Entities[Index]);
+				const FTestFragment_SparseInt* Fragment = EntityView.GetSparseFragmentDataPtr<FTestFragment_SparseInt>();
+				AITEST_NOT_NULL(TEXT("Fragment still exists after move"), Fragment);
+				AITEST_EQUAL(TEXT("Fragment value preserved"), Fragment->Value, Index * 100);
+			}
+
+			return true;
+		}
+	};
+	IMPLEMENT_AI_INSTANT_TEST(FSparseElements_DataPersistence, "System.Mass.Stress.Sparse.DataPersistence");
+
+	struct FSparseElements_CrossChunkBatchOperations : FEntityTestBase
+	{
+		virtual bool InstantTest() override
+		{
+			const int32 EntitiesPerChunk = EntityManager->DebugGetArchetypeEntitiesCountPerChunk(IntsArchetype);
+			constexpr int32 NumChunks = 5;
+			const int32 NumEntities = NumChunks * EntitiesPerChunk;
+
+			TArray<FMassEntityHandle> Entities;
+			EntityManager->BatchCreateEntities(IntsArchetype, NumEntities, Entities);
+
+			// Select entities scattered across all chunks
+			TArray<FMassEntityHandle> SelectedEntities;
+			for (int32 ChunkIdx = 0; ChunkIdx < NumChunks; ++ChunkIdx)
+			{
+				// Pick 3 entities from each chunk
+				for (int32 InChunkIdx = 0; InChunkIdx < 3 && (ChunkIdx * EntitiesPerChunk + InChunkIdx) < NumEntities; ++InChunkIdx)
+				{
+					SelectedEntities.Add(Entities[ChunkIdx * EntitiesPerChunk + InChunkIdx]);
+				}
+			}
+
+			// Batch add across chunks
+			FMassArchetypeEntityCollection Collection(IntsArchetype, SelectedEntities, FMassArchetypeEntityCollection::NoDuplicates);
+			EntityManager->BatchAddSparseElementToEntities({Collection}, FTestFragment_SparseInt::StaticStruct());
+
+			// Verify all have the element
+			for (const FMassEntityHandle& Entity : SelectedEntities)
+			{
+				FMassEntityView EntityView(*EntityManager, Entity);
+				AITEST_TRUE(TEXT("Entity has sparse element"), EntityView.HasElement<FTestFragment_SparseInt>());
+			}
+
+			// Batch remove across chunks
+			EntityManager->BatchRemoveSparseElementFromEntities({Collection}, FTestFragment_SparseInt::StaticStruct());
+
+			// Verify all removed
+			for (const FMassEntityHandle& Entity : SelectedEntities)
+			{
+				FMassEntityView EntityView(*EntityManager, Entity);
+				AITEST_FALSE(TEXT("Entity no longer has sparse element"), EntityView.HasElement<FTestFragment_SparseInt>());
+			}
+
+			return true;
+		}
+	};
+	IMPLEMENT_AI_INSTANT_TEST(FSparseElements_CrossChunkBatchOperations, "System.Mass.Stress.Sparse.CrossChunkBatch");
+
+	struct FSparseElements_Storage_ChunkBoundaryAccess : FEntityTestBase
+	{
+		virtual bool InstantTest() override
+		{
+			UE::Mass::FSparseElementsStorage Storage;
+
+			// Test accesses at chunk boundaries (default 128 elements per chunk)
+			constexpr int32 BoundaryIndices[] = {0, 127, 128, 255, 256, 383, 384};
+
+			TArray<FMassEntityHandle> Entities;
+			EntityManager->BatchCreateEntities(IntsArchetype, 500, Entities);
+
+			// Add sparse elements at boundary indices
+			for (int32 BoundaryIdx : BoundaryIndices)
+			{
+				if (BoundaryIdx < Entities.Num())
+				{
+					FStructView View = Storage.AddElementToEntity(Entities[BoundaryIdx], FTestFragment_SparseInt::StaticStruct());
+					AITEST_TRUE(TEXT("Boundary add succeeded"), View.IsValid());
+
+					FTestFragment_SparseInt* Fragment = View.GetPtr<FTestFragment_SparseInt>();
+					Fragment->Value = BoundaryIdx;
+				}
+			}
+
+			// Verify all boundary elements accessible
+			for (int32 BoundaryIdx : BoundaryIndices)
+			{
+				if (BoundaryIdx < Entities.Num())
+				{
+					FConstStructView View = Storage.GetElementDataForEntity(Entities[BoundaryIdx], FTestFragment_SparseInt::StaticStruct());
+					AITEST_TRUE(TEXT("Boundary element accessible"), View.IsValid());
+
+					const FTestFragment_SparseInt* Fragment = View.GetPtr<FTestFragment_SparseInt>();
+					AITEST_NOT_NULL(TEXT("Fragment pointer valid"), Fragment);
+					AITEST_EQUAL(TEXT("Boundary element value correct"), Fragment->Value, BoundaryIdx);
+				}
+			}
+
+			return true;
+		}
+	};
+	IMPLEMENT_AI_INSTANT_TEST(FSparseElements_Storage_ChunkBoundaryAccess, "System.Mass.Stress.Sparse.Storage.ChunkBoundary");
+
+	struct FSparseElements_BatchOperationsEmptyInput : FEntityTestBase
+	{
+		virtual bool InstantTest() override
+		{
+			TArray<FMassEntityHandle> EmptyArray;
+			FMassArchetypeEntityCollection EmptyCollection;
+
+			// Should not crash with empty input
+			EntityManager->BatchAddSparseElementToEntities({EmptyCollection}, FTestFragment_SparseInt::StaticStruct());
+			EntityManager->BatchRemoveSparseElementFromEntities({EmptyCollection}, FTestFragment_SparseInt::StaticStruct());
+
+			return true;
+		}
+	};
+	IMPLEMENT_AI_INSTANT_TEST(FSparseElements_BatchOperationsEmptyInput, "System.Mass.Stress.Sparse.BatchEmpty");
+
+	struct FSparseElements_ComplexQueryCombinations : FEntityTestBase
+	{
+		virtual bool InstantTest() override
+		{
+			constexpr int32 NumEntities = 20;
+			TArray<FMassEntityHandle> Entities;
+			EntityManager->BatchCreateEntities(IntsArchetype, NumEntities, Entities);
+
+			// Group 1: Has SparseInt only (entities 0-4)
+			for (int32 Index = 0; Index < 5; ++Index)
+			{
+				FStructView _ = EntityManager->AddSparseElementToEntity(Entities[Index], FTestFragment_SparseInt::StaticStruct());
+			}
+
+			// Group 2: Has SparseFloat only (entities 5-9)
+			for (int32 Index = 5; Index < 10; ++Index)
+			{
+				FStructView _ = EntityManager->AddSparseElementToEntity(Entities[Index], FTestFragment_SparseFloat::StaticStruct());
+			}
+
+			// Group 3: Has both (entities 10-14)
+			for (int32 Index = 10; Index < 15; ++Index)
+			{
+				FStructView _ = EntityManager->AddSparseElementToEntity(Entities[Index], FTestFragment_SparseInt::StaticStruct());
+				FStructView __ = EntityManager->AddSparseElementToEntity(Entities[Index], FTestFragment_SparseFloat::StaticStruct());
+			}
+
+			// Group 4: Has neither (entities 15-19)
+
+			// Query 1: Entities with both sparse elements
+			{
+				FMassEntityQuery Query(EntityManager);
+				// required since Queries don't support sparse-only requirements
+				Query.AddRequirement<FTestFragment_Int>(EMassFragmentAccess::ReadOnly);
+				Query.AddSparseRequirement<FTestFragment_SparseInt>();
+				Query.AddSparseRequirement<FTestFragment_SparseFloat>();
+
+				FMassExecutionContext ExecutionContext = EntityManager->CreateExecutionContext(0);
+				int32 Count = 0;
+				Query.ForEachEntityChunk(ExecutionContext, [&Count](FMassExecutionContext& Context)
+				{
+					for (FMassExecutionContext::FSparseEntityIterator It = Context.CreateSparseEntityIterator(); It; ++It)
+					{
+						++Count;
+					}
+				});
+
+				AITEST_EQUAL(TEXT("Query with both sparse elements"), Count, 5);
+			}
+
+			// Query 2: Entities with SparseInt but NOT SparseFloat
+			{
+				FMassEntityQuery Query(EntityManager);
+				// required since Queries don't support sparse-only requirements
+				Query.AddRequirement<FTestFragment_Int>(EMassFragmentAccess::ReadOnly);
+				Query.AddSparseRequirement<FTestFragment_SparseInt>();
+				Query.AddSparseRequirement<FTestFragment_SparseFloat>(EMassFragmentPresence::None);
+
+				FMassExecutionContext ExecutionContext = EntityManager->CreateExecutionContext(0);
+				int32 Count = 0;
+				Query.ForEachEntityChunk(ExecutionContext, [&Count](FMassExecutionContext& Context)
+				{
+					for (FMassExecutionContext::FSparseEntityIterator It = Context.CreateSparseEntityIterator(); It; ++It)
+					{
+						++Count;
+					}
+				});
+
+				AITEST_EQUAL(TEXT("Query with SparseInt, without SparseFloat"), Count, 5);
+			}
+
+			return true;
+		}
+	};
+	IMPLEMENT_AI_INSTANT_TEST(FSparseElements_ComplexQueryCombinations, "System.Mass.Stress.Sparse.ComplexQuery");
+
+	struct FSparseElements_MemoryFragmentation : FEntityTestBase
+	{
+		virtual bool InstantTest() override
+		{
+			UE::Mass::FSparseElementsStorage Storage;
+
+			constexpr int32 NumEntities = 1000;
+			constexpr int32 NumCycles = 10;
+
+			TArray<FMassEntityHandle> Entities;
+			EntityManager->BatchCreateEntities(IntsArchetype, NumEntities, Entities);
+
+			// Repeatedly add and remove to create fragmentation
+			for (int32 Cycle = 0; Cycle < NumCycles; ++Cycle)
+			{
+				// Add to all
+				for (const FMassEntityHandle& Entity : Entities)
+				{
+					FStructView _ = Storage.AddElementToEntity(Entity, FTestFragment_SparseInt::StaticStruct());
+				}
+
+				// Remove from half (creates holes)
+				for (int32 Index = 0; Index < NumEntities / 2; ++Index)
+				{
+					Storage.RemoveElementFromEntity<FTestFragment_SparseInt>(Entities[Index]);
+				}
+
+				// Re-add to different half
+				for (int32 Index = 0; Index < NumEntities / 2; ++Index)
+				{
+					FStructView _ = Storage.AddElementToEntity(Entities[Index], FTestFragment_SparseInt::StaticStruct());
+				}
+			}
+
+			// Final verification - all should have the element
+			for (const FMassEntityHandle& Entity : Entities)
+			{
+				FStructView View = Storage.GetMutableElementDataForEntity(Entity, FTestFragment_SparseInt::StaticStruct());
+				AITEST_TRUE(TEXT("Entity has sparse element after fragmentation"), View.IsValid());
+			}
+
+			return true;
+		}
+	};
+	IMPLEMENT_AI_INSTANT_TEST(FSparseElements_MemoryFragmentation, "System.Mass.Stress.Sparse.MemoryFragmentation");
+	
+	struct FSparseElements_OccupationMaskFiltering : FEntityTestBase
+	{
+		virtual bool InstantTest() override
+		{
+			constexpr int32 NumEntities = 200;
+			TArray<FMassEntityHandle> Entities;
+			EntityManager->BatchCreateEntities(IntsArchetype, NumEntities, Entities);
+
+			// Add to odd indices only
+			for (int32 Index = 1; Index < NumEntities; Index += 2)
+			{
+				FStructView _ = EntityManager->AddSparseElementToEntity(Entities[Index], FTestFragment_SparseInt::StaticStruct());
+			}
+
+			// Query should only find odd indices
+			FMassEntityQuery Query(EntityManager);
+			Query.AddRequirement<FTestFragment_Int>(EMassFragmentAccess::ReadOnly);
+			Query.AddSparseRequirement<FTestFragment_SparseInt>();
+
+			FMassExecutionContext ExecutionContext = EntityManager->CreateExecutionContext(0);
+			TSet<int32> FoundIndices;
+
+			Query.ForEachEntityChunk(ExecutionContext, [&FoundIndices](FMassExecutionContext& Context)
+			{
+				for (FMassExecutionContext::FSparseEntityIterator It = Context.CreateSparseEntityIterator(); It; ++It)
+				{
+					FoundIndices.Add(Context.GetEntities()[It].Index);
+				}
+			});
+
+			// Verify iterator correctly skipped even indices
+			int32 ExpectedCount = 0;
+			for (int32 Index = 0; Index < NumEntities; ++Index)
+			{
+				const bool bShouldHave = (Index % 2 == 1);
+				const bool bActuallyHas = FoundIndices.Contains(Entities[Index].Index);
+				AITEST_EQUAL(TEXT("Occupation mask filtering correct"), bActuallyHas, bShouldHave);
+				if (bShouldHave) ++ExpectedCount;
+			}
+
+			AITEST_EQUAL(TEXT("Found expected number of sparse elements"), FoundIndices.Num(), ExpectedCount);
+
+			return true;
+		}
+	};
+	IMPLEMENT_AI_INSTANT_TEST(FSparseElements_OccupationMaskFiltering, "System.Mass.Stress.Sparse.OccupationFiltering");
+	
+	// CRITICAL: Test that UStruct destructors are called (validates bug fix #2)
+	// Note: This test requires a special fragment type that tracks construction/destruction
+	// If FTestFragment_SparseInt doesn't have suitable tracking, this test may need modification
+	struct FSparseElements_DestructorValidation : FEntityTestBase
+	{
+		virtual bool InstantTest() override
+		{
+			UE::Mass::FSparseElementsStorage& Storage = EntityManager->DebugGetSparseElementsStorage();
+
+			constexpr int32 NumEntities = 10;
+			TArray<FMassEntityHandle> Entities;
+			EntityManager->BatchCreateEntities(IntsArchetype, NumEntities, Entities);
+
+			// Add sparse fragments - they should be properly constructed
+			for (const FMassEntityHandle& Entity : Entities)
+			{
+				FStructView View = EntityManager->AddSparseElementToEntity(Entity, FTestFragment_SparseInt::StaticStruct());
+				AITEST_TRUE(TEXT("Fragment added"), View.IsValid());
+
+				// Modify to ensure it's a real instance
+				FTestFragment_SparseInt* Fragment = View.GetPtr<FTestFragment_SparseInt>();
+				Fragment->Value = 42;
+			}
+
+			// Remove half - DestroyStruct should be called
+			for (int32 Index = 0; Index < NumEntities / 2; ++Index)
+			{
+				EntityManager->RemoveSparseElementFromEntity<FTestFragment_SparseInt>(Entities[Index]);
+
+				// Verify it's actually gone
+				FStructView View = Storage.GetMutableElementDataForEntity(Entities[Index], FTestFragment_SparseInt::StaticStruct());
+				AITEST_FALSE(TEXT("Fragment no longer exists"), View.IsValid());
+			}
+
+			// Destroy remaining entities - should call destructors via OnEntitiesDestroyed
+			TArray<FMassEntityHandle> RemainingEntities;
+			for (int32 Index = NumEntities / 2; Index < NumEntities; ++Index)
+			{
+				RemainingEntities.Add(Entities[Index]);
+			}
+			EntityManager->BatchDestroyEntities(RemainingEntities);
+
+			// Verify all destroyed
+			for (const FMassEntityHandle& Entity : RemainingEntities)
+			{
+				FStructView View = Storage.GetMutableElementDataForEntity(Entity, FTestFragment_SparseInt::StaticStruct());
+				AITEST_FALSE(TEXT("Destroyed entity fragment removed"), View.IsValid());
+			}
+
+			return true;
+		}
+	};
+	IMPLEMENT_AI_INSTANT_TEST(FSparseElements_DestructorValidation, "System.Mass.Stress.Sparse.DestructorValidation");
+	
+	struct FSparseElements_Storage_CountTracking : FEntityTestBase
+	{
+		virtual bool InstantTest() override
+		{
+			FSparseElementsStorage Storage;
+
+			// Initially empty
+			AITEST_EQUAL(TEXT("Initial count is zero"), Storage.GetNumElementsOfType<FTestFragment_SparseInt>(), 0u);
+			AITEST_FALSE(TEXT("Initially has no elements"), Storage.HasAnyElementsOfType<FTestFragment_SparseInt>());
+
+			constexpr int32 NumEntities = 50;
+			TArray<FMassEntityHandle> Entities;
+			EntityManager->BatchCreateEntities(IntsArchetype, NumEntities, Entities);
+
+			// Add to half
+			for (int32 i = 0; i < NumEntities / 2; ++i)
+			{
+				FStructView _ = Storage.AddElementToEntity(Entities[i], FTestFragment_SparseInt::StaticStruct());
+			}
+
+			AITEST_EQUAL(TEXT("Count after adding half"), Storage.GetNumElementsOfType<FTestFragment_SparseInt>(), 25u);
+			AITEST_TRUE(TEXT("Has elements after adding"), Storage.HasAnyElementsOfType<FTestFragment_SparseInt>());
+
+			// Remove quarter
+			for (int32 i = 0; i < NumEntities / 4; ++i)
+			{
+				Storage.RemoveElementFromEntity<FTestFragment_SparseInt>(Entities[i]);
+			}
+
+			AITEST_EQUAL(TEXT("Count after removing quarter"), Storage.GetNumElementsOfType<FTestFragment_SparseInt>(), 13u);
+
+			// Remove all
+			for (int32 i = 0; i < NumEntities; ++i)
+			{
+				Storage.RemoveElementFromEntity<FTestFragment_SparseInt>(Entities[i]);
+			}
+
+			AITEST_EQUAL(TEXT("Count after removing all"), Storage.GetNumElementsOfType<FTestFragment_SparseInt>(), 0u);
+			AITEST_FALSE(TEXT("No elements after removing all"), Storage.HasAnyElementsOfType<FTestFragment_SparseInt>());
+
+			return true;
+		}
+	};
+	IMPLEMENT_AI_INSTANT_TEST(FSparseElements_Storage_CountTracking, "System.Mass.Stress.Sparse.InstanceCountTracking");
+
+	struct FSparseElements_AddRemoveAddCycle : FEntityTestBase
+	{
+		virtual bool InstantTest() override
+		{
+			constexpr int32 NumEntities = 10;
+			TArray<FMassEntityHandle> Entities;
+			EntityManager->BatchCreateEntities(IntsArchetype, NumEntities, Entities);
+
+			for (int32 i = 0; i < NumEntities; ++i)
+			{
+				const FMassEntityHandle Entity = Entities[i];
+
+				// First Add: Set initial value
+				{
+					FStructView View = EntityManager->AddSparseElementToEntity(Entity, FTestFragment_SparseInt::StaticStruct());
+					AITEST_TRUE(TEXT("First add succeeded"), View.IsValid());
+
+					FTestFragment_SparseInt* Fragment = View.GetPtr<FTestFragment_SparseInt>();
+					Fragment->Value = 100 + i;
+				}
+
+				// Verify first add
+				{
+					FMassEntityView EntityView(*EntityManager, Entity);
+					const FTestFragment_SparseInt* Fragment = EntityView.GetSparseFragmentDataPtr<FTestFragment_SparseInt>();
+					AITEST_NOT_NULL(TEXT("Fragment exists after first add"), Fragment);
+					AITEST_EQUAL(TEXT("First add value correct"), Fragment->Value, 100 + i);
+				}
+
+				// Remove
+				EntityManager->RemoveSparseElementFromEntity(Entity, FTestFragment_SparseInt::StaticStruct());
+				
+				// Verify removal
+				{
+					FMassEntityView EntityView(*EntityManager, Entity);
+					const FTestFragment_SparseInt* Fragment = EntityView.GetSparseFragmentDataPtr<FTestFragment_SparseInt>();
+					AITEST_NULL(TEXT("Fragment removed"), Fragment);
+				}
+
+				// Second Add: Set different value to ensure fresh memory
+				{
+					FStructView View = EntityManager->AddSparseElementToEntity(Entity, FTestFragment_SparseInt::StaticStruct());
+					AITEST_TRUE(TEXT("Second add succeeded"), View.IsValid());
+
+					FTestFragment_SparseInt* Fragment = View.GetPtr<FTestFragment_SparseInt>();
+					// Value should be default-initialized (0), not the old value (100 + i)
+					AITEST_EQUAL(TEXT("Second add has fresh default value, not old data"), Fragment->Value, 0);
+
+					// Set new value
+					Fragment->Value = 200 + i;
+				}
+
+				// Verify second add has new value
+				{
+					FMassEntityView EntityView(*EntityManager, Entity);
+					const FTestFragment_SparseInt* Fragment = EntityView.GetSparseFragmentDataPtr<FTestFragment_SparseInt>();
+					AITEST_NOT_NULL(TEXT("Fragment exists after second add"), Fragment);
+					AITEST_EQUAL(TEXT("Second add value correct"), Fragment->Value, 200 + i);
+					AITEST_NOT_EQUAL(TEXT("Second add value different from first add"), Fragment->Value, 100 + i);
+				}
+			}
+
+			return true;
+		}
+	};
+	IMPLEMENT_AI_INSTANT_TEST(FSparseElements_AddRemoveAddCycle, "System.Mass.Stress.Sparse.AddRemoveAddCycle");
+
+	struct FSparseElements_AddRemoveAddAtChunkBoundaries : FEntityTestBase
+	{
+		virtual bool InstantTest() override
+		{
+			// Test at specific indices that cross chunk boundaries (default 128 per chunk)
+			constexpr int32 BoundaryIndices[] = { 0, 1, 126, 127, 128, 129, 255, 256, 383, 384 };
+			constexpr int32 MaxIndex = 500;
+
+			TArray<FMassEntityHandle> Entities;
+			EntityManager->BatchCreateEntities(IntsArchetype, MaxIndex, Entities);
+
+			for (int32 BoundaryIdx : BoundaryIndices)
+			{
+				if (BoundaryIdx >= Entities.Num())
+					continue;
+
+				const FMassEntityHandle Entity = Entities[BoundaryIdx];
+
+				// Cycle 1: Add with value A
+				{
+					FStructView View = EntityManager->AddSparseElementToEntity(Entity, FTestFragment_SparseFloat::StaticStruct());
+					AITEST_TRUE(TEXT("Boundary add 1 succeeded"), View.IsValid());
+
+					FTestFragment_SparseFloat* Fragment = View.GetPtr<FTestFragment_SparseFloat>();
+					Fragment->Value = static_cast<float>((BoundaryIdx + 1) * 10);
+				}
+
+				// Verify value A
+				{
+					FConstStructView View = EntityManager->GetSparseElementDataForEntity(
+						FTestFragment_SparseFloat::StaticStruct(), Entity);
+					AITEST_TRUE(TEXT("Fragment accessible at boundary"), View.IsValid());
+					AITEST_EQUAL(TEXT("Boundary value A correct"),
+						View.Get<FTestFragment_SparseFloat>().Value, static_cast<float>((BoundaryIdx + 1) * 10));
+				}
+
+				// Remove at boundary
+				EntityManager->RemoveSparseElementFromEntity(
+					Entity, FTestFragment_SparseFloat::StaticStruct());
+
+				// Verify removed
+				{
+					FConstStructView View = EntityManager->GetSparseElementDataForEntity(
+						FTestFragment_SparseFloat::StaticStruct(), Entity);
+					AITEST_FALSE(TEXT("Fragment removed at boundary"), View.IsValid());
+				}
+
+				// Cycle 2: Add with value B
+				{
+					FStructView View = EntityManager->AddSparseElementToEntity(Entity, FTestFragment_SparseFloat::StaticStruct());
+					AITEST_TRUE(TEXT("Boundary add 2 succeeded"), View.IsValid());
+
+					FTestFragment_SparseFloat* Fragment = View.GetPtr<FTestFragment_SparseFloat>();
+					// Should be default-initialized
+					AITEST_EQUAL(TEXT("Boundary second add has fresh default value"), Fragment->Value, 0.0f);
+
+					Fragment->Value = static_cast<float>((BoundaryIdx + 1) * 20); // Different value
+				}
+
+				// Verify value B (different from value A)
+				{
+					FConstStructView View = EntityManager->GetSparseElementDataForEntity(
+						FTestFragment_SparseFloat::StaticStruct(), Entity);
+					AITEST_TRUE(TEXT("Fragment accessible after second add"), View.IsValid());
+					AITEST_EQUAL(TEXT("Boundary value B correct"),
+						View.Get<FTestFragment_SparseFloat>().Value, static_cast<float>((BoundaryIdx + 1) * 20));
+					AITEST_NOT_EQUAL(TEXT("Boundary value B different from A"),
+						View.Get<FTestFragment_SparseFloat>().Value, static_cast<float>((BoundaryIdx + 1) * 10));
+				}
+			}
+
+			return true;
+		}
+	};
+	IMPLEMENT_AI_INSTANT_TEST(FSparseElements_AddRemoveAddAtChunkBoundaries, "System.Mass.Stress.Sparse.AddRemoveAddBoundaries");
+
+	struct FSparseElements_AddRemoveAddWithTags : FEntityTestBase
+	{
+		virtual bool InstantTest() override
+		{
+			constexpr int32 NumEntities = 5;
+			TArray<FMassEntityHandle> Entities;
+			EntityManager->BatchCreateEntities(FloatsArchetype, NumEntities, Entities);
+
+			for (const FMassEntityHandle& Entity : Entities)
+			{
+				// Add sparse tag
+				{
+					FStructView View = EntityManager->AddSparseElementToEntity(Entity, FTestTag_SparseA::StaticStruct());
+					// Tags return invalid view (no data)
+					AITEST_FALSE(TEXT("Tag add returns empty view"), View.IsValid());
+				}
+
+				// Verify tag present
+				{
+					FMassEntityView EntityView(*EntityManager, Entity);
+					AITEST_TRUE(TEXT("Tag present after add"), EntityView.HasElement<FTestTag_SparseA>());
+				}
+
+				// Remove tag
+				EntityManager->RemoveSparseElementFromEntity(Entity, FTestTag_SparseA::StaticStruct());
+
+				// Verify tag removed
+				{
+					FMassEntityView EntityView(*EntityManager, Entity);
+					AITEST_FALSE(TEXT("Tag removed"), EntityView.HasElement<FTestTag_SparseA>());
+				}
+
+				// Add tag again
+				{
+					FStructView View = EntityManager->AddSparseElementToEntity(Entity, FTestTag_SparseA::StaticStruct());
+					AITEST_FALSE(TEXT("Tag second add returns empty view"), View.IsValid());
+				}
+
+				// Verify tag present again
+				{
+					FMassEntityView EntityView(*EntityManager, Entity);
+					AITEST_TRUE(TEXT("Tag present after second add"), EntityView.HasElement<FTestTag_SparseA>());
+				}
+			}
+
+			return true;
+		}
+	};
+	IMPLEMENT_AI_INSTANT_TEST(FSparseElements_AddRemoveAddWithTags, "System.Mass.Stress.Sparse.AddRemoveAddTags");
+
+	struct FSparseElements_AddRemoveAddMultipleTypes : FEntityTestBase
+	{
+		virtual bool InstantTest() override
+		{
+			const FMassEntityHandle Entity = EntityManager->CreateEntity(IntsArchetype);
+
+			// Add two different sparse types
+			{
+				FStructView IntView = EntityManager->AddSparseElementToEntity(Entity, FTestFragment_SparseInt::StaticStruct());
+				IntView.Get<FTestFragment_SparseInt>().Value = 100;
+
+				FStructView FloatView = EntityManager->AddSparseElementToEntity(Entity, FTestFragment_SparseFloat::StaticStruct());
+				FloatView.Get<FTestFragment_SparseFloat>().Value = 50.0f;
+			}
+
+			// Verify both present
+			{
+				FMassEntityView EntityView(*EntityManager, Entity);
+				AITEST_TRUE(TEXT("Has SparseInt"), EntityView.HasElement<FTestFragment_SparseInt>());
+				AITEST_TRUE(TEXT("Has SparseFloat"), EntityView.HasElement<FTestFragment_SparseFloat>());
+			}
+
+			// Remove only SparseInt
+			EntityManager->RemoveSparseElementFromEntity(Entity, FTestFragment_SparseInt::StaticStruct());
+
+			// Verify SparseInt removed but SparseFloat remains
+			{
+				FMassEntityView EntityView(*EntityManager, Entity);
+				AITEST_FALSE(TEXT("SparseInt removed"), EntityView.HasElement<FTestFragment_SparseInt>());
+				AITEST_TRUE(TEXT("SparseFloat still present"), EntityView.HasElement<FTestFragment_SparseFloat>());
+
+				const FTestFragment_SparseFloat* FloatFragment = EntityView.GetSparseFragmentDataPtr<FTestFragment_SparseFloat>();
+				AITEST_NOT_NULL(TEXT("SparseFloat accessible"), FloatFragment);
+				AITEST_EQUAL(TEXT("SparseFloat value unchanged"), FloatFragment->Value, 50.0f);
+			}
+
+			// Add SparseInt again with different value
+			{
+				FStructView IntView = EntityManager->AddSparseElementToEntity(Entity, FTestFragment_SparseInt::StaticStruct());
+				AITEST_TRUE(TEXT("SparseInt re-added"), IntView.IsValid());
+
+				FTestFragment_SparseInt* IntFragment = IntView.GetPtr<FTestFragment_SparseInt>();
+				AITEST_EQUAL(TEXT("Re-added SparseInt has fresh default value"), IntFragment->Value, 0);
+
+				IntFragment->Value = 200; // Different from original 100
+			}
+
+			// Verify both present with correct values
+			{
+				FMassEntityView EntityView(*EntityManager, Entity);
+				AITEST_TRUE(TEXT("Both types present"),
+					EntityView.HasElement<FTestFragment_SparseInt>() &&
+					EntityView.HasElement<FTestFragment_SparseFloat>());
+
+				const FTestFragment_SparseInt* IntFragment = EntityView.GetSparseFragmentDataPtr<FTestFragment_SparseInt>();
+				const FTestFragment_SparseFloat* FloatFragment = EntityView.GetSparseFragmentDataPtr<FTestFragment_SparseFloat>();
+
+				AITEST_NOT_NULL(TEXT("SparseInt accessible"), IntFragment);
+				AITEST_NOT_NULL(TEXT("SparseFloat accessible"), FloatFragment);
+
+				AITEST_EQUAL(TEXT("SparseInt has new value"), IntFragment->Value, 200);
+				AITEST_NOT_EQUAL(TEXT("SparseInt different from original"), IntFragment->Value, 100);
+				AITEST_EQUAL(TEXT("SparseFloat value unchanged"), FloatFragment->Value, 50.0f);
+			}
+
+			return true;
+		}
+	};
+	IMPLEMENT_AI_INSTANT_TEST(FSparseElements_AddRemoveAddMultipleTypes, "System.Mass.Stress.Sparse.AddRemoveAddMultipleTypes");
+
+	struct FSparseElements_AddRemoveAddWithArchetypeChange : FEntityTestBase
+	{
+		virtual bool InstantTest() override
+		{
+			const FMassEntityHandle Entity = EntityManager->CreateEntity(IntsArchetype);
+
+			// Add sparse element to entity in IntsArchetype
+			{
+				FStructView View = EntityManager->AddSparseElementToEntity(Entity, FTestFragment_SparseInt::StaticStruct());
+				View.Get<FTestFragment_SparseInt>().Value = 111;
+			}
+
+			// Verify
+			{
+				AITEST_EQUAL(TEXT("Entity in IntsArchetype"),
+					EntityManager->GetArchetypeForEntity(Entity), IntsArchetype);
+
+				FMassEntityView EntityView(*EntityManager, Entity);
+				const FTestFragment_SparseInt* Fragment = EntityView.GetSparseFragmentDataPtr<FTestFragment_SparseInt>();
+				AITEST_NOT_NULL(TEXT("Sparse element present"), Fragment);
+				AITEST_EQUAL(TEXT("Value correct"), Fragment->Value, 111);
+			}
+
+			// Remove sparse element
+			{
+				EntityManager->RemoveSparseElementFromEntity(Entity, FTestFragment_SparseInt::StaticStruct());
+			}
+
+			// Change archetype by adding regular fragment
+			{
+				EntityManager->AddFragmentToEntity(Entity, FTestFragment_Float::StaticStruct());
+				AITEST_EQUAL(TEXT("Entity moved to FloatsIntsArchetype"),
+					EntityManager->GetArchetypeForEntity(Entity), FloatsIntsArchetype);
+			}
+
+			// Add sparse element again (now in different archetype)
+			{
+				FStructView View = EntityManager->AddSparseElementToEntity(Entity, FTestFragment_SparseInt::StaticStruct());
+				AITEST_TRUE(TEXT("Sparse element added after archetype change"), View.IsValid());
+
+				FTestFragment_SparseInt* Fragment = View.GetPtr<FTestFragment_SparseInt>();
+				AITEST_EQUAL(TEXT("Fresh default value after archetype change"), Fragment->Value, 0);
+
+				Fragment->Value = 222;
+			}
+
+			// Verify in new archetype
+			{
+				AITEST_EQUAL(TEXT("Still in FloatsIntsArchetype"),
+					EntityManager->GetArchetypeForEntity(Entity), FloatsIntsArchetype);
+
+				FMassEntityView EntityView(*EntityManager, Entity);
+				const FTestFragment_SparseInt* Fragment = EntityView.GetSparseFragmentDataPtr<FTestFragment_SparseInt>();
+				AITEST_NOT_NULL(TEXT("Sparse element present in new archetype"), Fragment);
+				AITEST_EQUAL(TEXT("New value correct"), Fragment->Value, 222);
+				AITEST_NOT_EQUAL(TEXT("Value different from original"), Fragment->Value, 111);
+			}
+
+			return true;
+		}
+	};
+	IMPLEMENT_AI_INSTANT_TEST(FSparseElements_AddRemoveAddWithArchetypeChange, "System.Mass.Stress.Sparse.AddRemoveAddWithArchetypeChange");
+
+	struct FSparseElements_BatchAddFragmentInstances_PureSparse : FEntityTestBase
+	{
+		virtual bool InstantTest() override
+		{
+			constexpr int32 NumEntities = 20;
+			TArray<FMassEntityHandle> Entities;
+			EntityManager->BatchCreateEntities(IntsArchetype, NumEntities, Entities);
+
+			// Create payload with sparse fragment instances
+			UE::Mass::TMultiArray<FTestFragment_SparseInt, FTestFragment_SparseFloat> Payload;
+			for (int32 i = 0; i < NumEntities; ++i)
+			{
+				FTestFragment_SparseInt SparseInt;
+				SparseInt.Value = i * 100;
+
+				FTestFragment_SparseFloat SparseFloat;
+				SparseFloat.Value = static_cast<float>(i * 5);
+
+				Payload.Add(SparseInt, SparseFloat);
+			}
+
+			// Create entity collections with payload
+			TArray<FStructArrayView> GenericPayload;
+			GenericPayload.Reserve(Payload.GetNumArrays());
+			Payload.GetAsGenericMultiArray(GenericPayload);
+
+			TArray<FMassArchetypeEntityCollectionWithPayload> EntityCollections;
+			FMassArchetypeEntityCollectionWithPayload::CreateEntityRangesWithPayload(
+				*EntityManager, Entities, FMassArchetypeEntityCollection::FoldDuplicates,
+				FMassGenericPayloadView(GenericPayload), EntityCollections);
+
+			// Build fragments bitset
+			FMassFragmentBitSet FragmentsAffected;
+			FragmentsAffected.Add(FTestFragment_SparseInt::StaticStruct());
+			FragmentsAffected.Add(FTestFragment_SparseFloat::StaticStruct());
+
+			// Execute batch add
+			EntityManager->BatchAddFragmentInstancesForEntities(EntityCollections, FragmentsAffected);
+
+			// Verify all entities have sparse fragments with correct data
+			for (int32 i = 0; i < NumEntities; ++i)
+			{
+				FMassEntityView EntityView(*EntityManager, Entities[i]);
+
+				const FTestFragment_SparseInt* SparseInt = EntityView.GetSparseFragmentDataPtr<FTestFragment_SparseInt>();
+				const FTestFragment_SparseFloat* SparseFloat = EntityView.GetSparseFragmentDataPtr<FTestFragment_SparseFloat>();
+
+				AITEST_NOT_NULL(TEXT("SparseInt added"), SparseInt);
+				AITEST_NOT_NULL(TEXT("SparseFloat added"), SparseFloat);
+
+				AITEST_EQUAL(TEXT("SparseInt value correct"), SparseInt->Value, i * 100);
+				AITEST_EQUAL(TEXT("SparseFloat value correct"), SparseFloat->Value, static_cast<float>(i * 5));
+			}
+
+			return true;
+		}
+	};
+	IMPLEMENT_AI_INSTANT_TEST(FSparseElements_BatchAddFragmentInstances_PureSparse, "System.Mass.Stress.Sparse.BatchAddFragmentInstances.PureSparse");
+
+	struct FSparseElements_BatchAddFragmentInstances_MixedRegularAndSparse : FEntityTestBase
+	{
+		virtual bool InstantTest() override
+		{
+			constexpr int32 NumEntities = 15;
+			TArray<FMassEntityHandle> Entities;
+			EntityManager->BatchCreateEntities(IntsArchetype, NumEntities, Entities);
+
+			// Create payload with both regular and sparse fragments
+			UE::Mass::TMultiArray<FTestFragment_Float, FTestFragment_SparseInt> Payload;
+			for (int32 i = 0; i < NumEntities; ++i)
+			{
+				FTestFragment_Float RegularFloat;
+				RegularFloat.Value = static_cast<float>(i * 10);
+
+				FTestFragment_SparseInt SparseInt;
+				SparseInt.Value = i * 200;
+
+				Payload.Add(RegularFloat, SparseInt);
+			}
+
+			// Create entity collections with payload
+			TArray<FStructArrayView> GenericPayload;
+			GenericPayload.Reserve(Payload.GetNumArrays());
+			Payload.GetAsGenericMultiArray(GenericPayload);
+
+			TArray<FMassArchetypeEntityCollectionWithPayload> EntityCollections;
+			FMassArchetypeEntityCollectionWithPayload::CreateEntityRangesWithPayload(
+				*EntityManager, Entities, FMassArchetypeEntityCollection::FoldDuplicates,
+				FMassGenericPayloadView(GenericPayload), EntityCollections);
+
+			// Build fragments bitset (mixed regular + sparse)
+			FMassFragmentBitSet FragmentsAffected;
+			FragmentsAffected.Add(FTestFragment_Float::StaticStruct());
+			FragmentsAffected.Add(FTestFragment_SparseInt::StaticStruct());
+
+			// Execute batch add
+			EntityManager->BatchAddFragmentInstancesForEntities(EntityCollections, FragmentsAffected);
+
+			// Verify all entities have both regular and sparse fragments
+			for (int32 i = 0; i < NumEntities; ++i)
+			{
+				FMassEntityView EntityView(*EntityManager, Entities[i]);
+
+				// Regular fragment should be in archetype
+				const FTestFragment_Float* RegularFloat = EntityView.GetFragmentDataPtr<FTestFragment_Float>();
+				AITEST_NOT_NULL(TEXT("Regular float added"), RegularFloat);
+				AITEST_EQUAL(TEXT("Regular float value correct"), RegularFloat->Value, static_cast<float>(i * 10));
+
+				// Sparse fragment should be in external storage
+				const FTestFragment_SparseInt* SparseInt = EntityView.GetSparseFragmentDataPtr<FTestFragment_SparseInt>();
+				AITEST_NOT_NULL(TEXT("Sparse int added"), SparseInt);
+				AITEST_EQUAL(TEXT("Sparse int value correct"), SparseInt->Value, i * 200);
+			}
+
+			return true;
+		}
+	};
+	IMPLEMENT_AI_INSTANT_TEST(FSparseElements_BatchAddFragmentInstances_MixedRegularAndSparse, "System.Mass.Stress.Sparse.BatchAddFragmentInstances.MixedRegularAndSparse");
+
+	struct FSparseElements_BatchAddFragmentInstances_CrossChunks : FEntityTestBase
+	{
+		virtual bool InstantTest() override
+		{
+			const int32 EntitiesPerChunk = EntityManager->DebugGetArchetypeEntitiesCountPerChunk(FloatsArchetype);
+			const int32 NumEntities = static_cast<int32>(static_cast<float>(EntitiesPerChunk) * 2.5f); // Span 3 chunks
+
+			TArray<FMassEntityHandle> Entities;
+			EntityManager->BatchCreateEntities(FloatsArchetype, NumEntities, Entities);
+
+			// Create payload spanning multiple chunks
+			UE::Mass::TMultiArray<FTestFragment_SparseInt> Payload;
+			for (int32 i = 0; i < NumEntities; ++i)
+			{
+				FTestFragment_SparseInt SparseInt;
+				SparseInt.Value = i * 777;
+				Payload.Add(SparseInt);
+			}
+
+			// Create entity collections
+			TArray<FStructArrayView> GenericPayload;
+			GenericPayload.Reserve(Payload.GetNumArrays());
+			Payload.GetAsGenericMultiArray(GenericPayload);
+
+			TArray<FMassArchetypeEntityCollectionWithPayload> EntityCollections;
+			FMassArchetypeEntityCollectionWithPayload::CreateEntityRangesWithPayload(
+				*EntityManager, Entities, FMassArchetypeEntityCollection::FoldDuplicates,
+				FMassGenericPayloadView(GenericPayload), EntityCollections);
+
+			FMassFragmentBitSet FragmentsAffected;
+			FragmentsAffected.Add(FTestFragment_SparseInt::StaticStruct());
+
+			// Execute across chunks
+			EntityManager->BatchAddFragmentInstancesForEntities(EntityCollections, FragmentsAffected);
+
+			// Verify all entities across all chunks have correct data
+			for (int32 i = 0; i < NumEntities; ++i)
+			{
+				FMassEntityView EntityView(*EntityManager, Entities[i]);
+				const FTestFragment_SparseInt* SparseInt = EntityView.GetSparseFragmentDataPtr<FTestFragment_SparseInt>();
+
+				AITEST_NOT_NULL(TEXT("Sparse fragment present across chunks"), SparseInt);
+				AITEST_EQUAL(TEXT("Sparse fragment value correct across chunks"), SparseInt->Value, i * 777);
+			}
+
+			return true;
+		}
+	};
+	IMPLEMENT_AI_INSTANT_TEST(FSparseElements_BatchAddFragmentInstances_CrossChunks, "System.Mass.Stress.Sparse.BatchAddFragmentInstances.CrossChunks");
+
+	struct FSparseElements_BatchAddFragmentInstances_OverwriteExisting : FEntityTestBase
+	{
+		virtual bool InstantTest() override
+		{
+			constexpr int32 NumEntities = 10;
+			TArray<FMassEntityHandle> Entities;
+			EntityManager->BatchCreateEntities(IntsArchetype, NumEntities, Entities);
+
+			// First, add sparse fragments with initial values
+			for (int32 i = 0; i < NumEntities; ++i)
+			{
+				FStructView SparseInt = EntityManager->AddSparseElementToEntity(Entities[i], FTestFragment_SparseInt::StaticStruct());
+				SparseInt.Get<FTestFragment_SparseInt>().Value = (i + 1) * 10; // Initial value
+			}
+
+			// Verify initial values
+			for (int32 i = 0; i < NumEntities; ++i)
+			{
+				FMassEntityView EntityView(*EntityManager, Entities[i]);
+				const FTestFragment_SparseInt* SparseInt = EntityView.GetSparseFragmentDataPtr<FTestFragment_SparseInt>();
+				AITEST_NOT_NULL(TEXT("Initial sparse fragment present"), SparseInt);
+				AITEST_EQUAL(TEXT("Initial value correct"), SparseInt->Value, (i + 1) * 10);
+			}
+
+			// Now use BatchAddFragmentInstancesForEntities to overwrite with new values
+			UE::Mass::TMultiArray<FTestFragment_SparseInt> Payload;
+			for (int32 i = 0; i < NumEntities; ++i)
+			{
+				FTestFragment_SparseInt NewSparseInt;
+				NewSparseInt.Value = (i + 1) * 999; // New value
+				Payload.Add(NewSparseInt);
+			}
+
+			TArray<FStructArrayView> GenericPayload;
+			GenericPayload.Reserve(Payload.GetNumArrays());
+			Payload.GetAsGenericMultiArray(GenericPayload);
+
+			TArray<FMassArchetypeEntityCollectionWithPayload> EntityCollections;
+			FMassArchetypeEntityCollectionWithPayload::CreateEntityRangesWithPayload(
+				*EntityManager, Entities, FMassArchetypeEntityCollection::FoldDuplicates,
+				FMassGenericPayloadView(GenericPayload), EntityCollections);
+
+			FMassFragmentBitSet FragmentsAffected;
+			FragmentsAffected.Add(FTestFragment_SparseInt::StaticStruct());
+
+			EntityManager->BatchAddFragmentInstancesForEntities(EntityCollections, FragmentsAffected);
+
+			// Verify values were overwritten
+			for (int32 i = 0; i < NumEntities; ++i)
+			{
+				FMassEntityView EntityView(*EntityManager, Entities[i]);
+				const FTestFragment_SparseInt* SparseInt = EntityView.GetSparseFragmentDataPtr<FTestFragment_SparseInt>();
+
+				AITEST_NOT_NULL(TEXT("Sparse fragment still present"), SparseInt);
+				AITEST_EQUAL(TEXT("Value overwritten correctly"), SparseInt->Value, (i + 1) * 999);
+				AITEST_NOT_EQUAL(TEXT("Value changed from initial"), SparseInt->Value, (i + 1) * 10);
+			}
+
+			return true;
+		}
+	};
+	IMPLEMENT_AI_INSTANT_TEST(FSparseElements_BatchAddFragmentInstances_OverwriteExisting, "System.Mass.Stress.Sparse.BatchAddFragmentInstances.OverwriteExisting");
+
+	struct FSparseElements_CommandAddFragmentInstances_PureSparse : FEntityTestBase
+	{
+		virtual bool InstantTest() override
+		{
+			constexpr int32 NumEntities = 12;
+			TArray<FMassEntityHandle> Entities;
+			EntityManager->BatchCreateEntities(FloatsArchetype, NumEntities, Entities);
+
+			// Use command to add sparse fragments
+			for (int32 i = 0; i < NumEntities; ++i)
+			{
+				FTestFragment_SparseInt SparseInt;
+				SparseInt.Value = i * 50;
+
+				FTestFragment_SparseFloat SparseFloat;
+				SparseFloat.Value = static_cast<float>(i * 3);
+
+				EntityManager->Defer().PushCommand<FMassCommandAddFragmentInstances>(Entities[i], SparseInt, SparseFloat);
+			}
+
+			// Flush commands
+			EntityManager->FlushCommands();
+
+			// Verify all entities have sparse fragments with correct values
+			for (int32 i = 0; i < NumEntities; ++i)
+			{
+				FMassEntityView EntityView(*EntityManager, Entities[i]);
+
+				const FTestFragment_SparseInt* SparseInt = EntityView.GetSparseFragmentDataPtr<FTestFragment_SparseInt>();
+				const FTestFragment_SparseFloat* SparseFloat = EntityView.GetSparseFragmentDataPtr<FTestFragment_SparseFloat>();
+
+				AITEST_NOT_NULL(TEXT("Command added SparseInt"), SparseInt);
+				AITEST_NOT_NULL(TEXT("Command added SparseFloat"), SparseFloat);
+
+				AITEST_EQUAL(TEXT("Command SparseInt value correct"), SparseInt->Value, i * 50);
+				AITEST_EQUAL(TEXT("Command SparseFloat value correct"), SparseFloat->Value, static_cast<float>(i * 3));
+			}
+
+			return true;
+		}
+	};
+	IMPLEMENT_AI_INSTANT_TEST(FSparseElements_CommandAddFragmentInstances_PureSparse, "System.Mass.Stress.Sparse.CommandAddFragmentInstances.PureSparse");
+
+	struct FSparseElements_CommandAddFragmentInstances_MixedRegularAndSparse : FEntityTestBase
+	{
+		virtual bool InstantTest() override
+		{
+			constexpr int32 NumEntities = 8;
+			TArray<FMassEntityHandle> Entities;
+			EntityManager->BatchCreateEntities(IntsArchetype, NumEntities, Entities);
+
+			// Use command with mixed regular + sparse fragments
+			for (int32 i = 0; i < NumEntities; ++i)
+			{
+				FTestFragment_Float RegularFloat;
+				RegularFloat.Value = static_cast<float>(i * 25);
+
+				FTestFragment_SparseInt SparseInt;
+				SparseInt.Value = i * 888;
+
+				EntityManager->Defer().PushCommand<FMassCommandAddFragmentInstances>(Entities[i], RegularFloat, SparseInt);
+			}
+
+			EntityManager->FlushCommands();
+
+			// Verify both regular and sparse fragments
+			for (int32 i = 0; i < NumEntities; ++i)
+			{
+				FMassEntityView EntityView(*EntityManager, Entities[i]);
+
+				const FTestFragment_Float* RegularFloat = EntityView.GetFragmentDataPtr<FTestFragment_Float>();
+				const FTestFragment_SparseInt* SparseInt = EntityView.GetSparseFragmentDataPtr<FTestFragment_SparseInt>();
+
+				AITEST_NOT_NULL(TEXT("Command added regular float"), RegularFloat);
+				AITEST_NOT_NULL(TEXT("Command added sparse int"), SparseInt);
+
+				AITEST_EQUAL(TEXT("Command regular float value correct"), RegularFloat->Value, static_cast<float>(i * 25));
+				AITEST_EQUAL(TEXT("Command sparse int value correct"), SparseInt->Value, i * 888);
+			}
+
+			return true;
+		}
+	};
+	IMPLEMENT_AI_INSTANT_TEST(FSparseElements_CommandAddFragmentInstances_MixedRegularAndSparse, "System.Mass.Stress.Sparse.CommandAddFragmentInstances.MixedRegularAndSparse");
+
+	struct FSparseElements_CommandAddFragmentInstances_CrossChunks : FEntityTestBase
+	{
+		virtual bool InstantTest() override
+		{
+			const int32 EntitiesPerChunk = EntityManager->DebugGetArchetypeEntitiesCountPerChunk(IntsArchetype);
+			const int32 NumChunks = 4;
+			const int32 NumEntities = EntitiesPerChunk * NumChunks;
+
+			TArray<FMassEntityHandle> Entities;
+			EntityManager->BatchCreateEntities(IntsArchetype, NumEntities, Entities);
+
+			// Use command across multiple chunks
+			for (int32 i = 0; i < NumEntities; ++i)
+			{
+				FTestFragment_SparseFloat SparseFloat;
+				SparseFloat.Value = static_cast<float>(i * 11);
+				EntityManager->Defer().PushCommand<FMassCommandAddFragmentInstances>(Entities[i], SparseFloat);
+			}
+
+			EntityManager->FlushCommands();
+
+			// Verify across all chunks
+			int32 ChunksFound = 0;
+			int32 LastChunkIndex = -1;
+
+			for (int32 i = 0; i < NumEntities; ++i)
+			{
+				FMassEntityView EntityView(*EntityManager, Entities[i]);
+				const FTestFragment_SparseFloat* SparseFloat = EntityView.GetSparseFragmentDataPtr<FTestFragment_SparseFloat>();
+
+				AITEST_NOT_NULL(TEXT("Sparse fragment present across chunks via command"), SparseFloat);
+				AITEST_EQUAL(TEXT("Value correct across chunks via command"), SparseFloat->Value, static_cast<float>(i * 11));
+
+				// Track chunk transitions
+				int32 CurrentChunkIndex = i / EntitiesPerChunk;
+				if (CurrentChunkIndex != LastChunkIndex)
+				{
+					++ChunksFound;
+					LastChunkIndex = CurrentChunkIndex;
+				}
+			}
+
+			AITEST_EQUAL(TEXT("Command processed all chunks"), ChunksFound, NumChunks);
+
+			return true;
+		}
+	};
+	IMPLEMENT_AI_INSTANT_TEST(FSparseElements_CommandAddFragmentInstances_CrossChunks, "System.Mass.Stress.Sparse.CommandAddFragmentInstances.CrossChunks");
+
+	struct FSparseElements_CommandAddFragmentInstances_PartialBatch : FEntityTestBase
+	{
+		virtual bool InstantTest() override
+		{
+			constexpr int32 NumEntities = 30;
+			TArray<FMassEntityHandle> Entities;
+			EntityManager->BatchCreateEntities(FloatsArchetype, NumEntities, Entities);
+
+			// Only add sparse fragments to every 3rd entity via command
+			int32 ExpectedCount = 0;
+			for (int32 i = 0; i < NumEntities; i += 3)
+			{
+				FTestFragment_SparseInt SparseInt;
+				SparseInt.Value = i * 33;
+				EntityManager->Defer().PushCommand<FMassCommandAddFragmentInstances>(Entities[i], SparseInt);
+				++ExpectedCount;
+			}
+
+			EntityManager->FlushCommands();
+
+			// Verify only targeted entities have sparse fragments
+			int32 ActualCount = 0;
+			for (int32 i = 0; i < NumEntities; ++i)
+			{
+				FMassEntityView EntityView(*EntityManager, Entities[i]);
+				const FTestFragment_SparseInt* SparseInt = EntityView.GetSparseFragmentDataPtr<FTestFragment_SparseInt>();
+
+				if (i % 3 == 0)
+				{
+					AITEST_NOT_NULL(TEXT("Targeted entity has sparse fragment"), SparseInt);
+					AITEST_EQUAL(TEXT("Targeted entity value correct"), SparseInt->Value, i * 33);
+					++ActualCount;
+				}
+				else
+				{
+					AITEST_NULL(TEXT("Non-targeted entity has no sparse fragment"), SparseInt);
+				}
+			}
+
+			AITEST_EQUAL(TEXT("Correct number of entities received sparse fragments"), ActualCount, ExpectedCount);
+
+			return true;
+		}
+	};
+	IMPLEMENT_AI_INSTANT_TEST(FSparseElements_CommandAddFragmentInstances_PartialBatch, "System.Mass.Stress.Sparse.CommandAddFragmentInstances.PartialBatch");
+
+	//struct FSparseElements_Storage_Statistics : FEntityTestBase
+	//{
+	//	virtual bool InstantTest() override
+	//	{
+	//		FSparseElementsStorage Storage;
+
+	//		constexpr int32 NumEntities = 200;
+	//		TArray<FMassEntityHandle> Entities;
+	//		EntityManager->BatchCreateEntities(IntsArchetype, NumEntities, Entities);
+
+	//		// Add elements across multiple chunks (default 128/chunk)
+	//		for (int32 i = 0; i < NumEntities; i += 2)  // Every other entity
+	//		{
+	//			Storage.AddElementToEntity<FTestFragment_SparseInt>(Entities[i]);
+	//		}
+
+	//		// Get type stats
+	//		auto Stats = Storage.GetTypeStats<FTestFragment_SparseInt>();
+	//		AITEST_EQUAL(TEXT("Stats reports correct count"), Stats.NumElements, 100u);
+	//		AITEST_TRUE(TEXT("Stats reports multiple chunks"), Stats.NumChunks >= 2);
+	//		AITEST_TRUE(TEXT("Occupation rate is reasonable"), Stats.OccupationRate > 0.0f && Stats.OccupationRate < 1.0f);
+	//		AITEST_TRUE(TEXT("Memory usage reported"), Stats.MemoryUsed > 0);
+
+	//		// Get storage stats
+	//		auto StorageStats = Storage.GetStorageStats();
+	//		AITEST_EQUAL(TEXT("Storage stats total elements"), StorageStats.TotalElements, 100u);
+	//		AITEST_TRUE(TEXT("Storage stats active pools"), StorageStats.NumActiveTypePools >= 1);
+
+	//		return true;
+	//	}
+	//};
+	//IMPLEMENT_AI_INSTANT_TEST(FSparseElements_Storage_CountTracking, "System.Mass.Stress.Sparse.StorageStatistics");
+
+#endif // WITH_MASSENTITY_DEBUG
+
+//-----------------------------------------------------------------------------
+// Chunk and Archetype Bitset Tracking Tests
+//-----------------------------------------------------------------------------
+
+struct FSparseElements_ChunkBitset_AddSingle : FEntityTestBase
+{
+	virtual bool InstantTest() override
+	{
+		CA_ASSUME(EntityManager);
+
+		constexpr int32 NumEntities = 10;
+		TArray<FMassEntityHandle> Entities;
+		EntityManager->BatchCreateEntities(IntsArchetype, NumEntities, Entities);
+
+		// Add sparse element to first entity
+		FStructView _ = EntityManager->AddSparseElementToEntity(Entities[0], FTestFragment_SparseInt::StaticStruct());
+
+		// Verify chunk bitset is updated
+		FMassArchetypeData& Archetype = FMassArchetypeHelper::ArchetypeDataFromHandleChecked(IntsArchetype);
+		const FMassArchetypeChunk& Chunk = Archetype.DebugGetChunk(0);
+		AITEST_TRUE("Chunk should have sparse elements", Chunk.HasSparseElements());
+
+		const UE::Mass::FChunkSparseElements& ChunkSparseElements = Chunk.GetSparseElementsUnsafe();
+		AITEST_TRUE("Chunk bitset should contain sparse int",
+			ChunkSparseElements.GetSparseElementsPresent().Contains(FTestFragment_SparseInt::StaticStruct()));
+
+		return true;
+	}
+};
+IMPLEMENT_AI_INSTANT_TEST(FSparseElements_ChunkBitset_AddSingle, "System.Mass.Stress.Sparse.ChunkBitset.AddSingle");
+
+struct FSparseElements_ChunkBitset_RemoveSingle : FEntityTestBase
+{
+	virtual bool InstantTest() override
+	{
+		CA_ASSUME(EntityManager);
+
+		constexpr int32 NumEntities = 10;
+		TArray<FMassEntityHandle> Entities;
+		EntityManager->BatchCreateEntities(IntsArchetype, NumEntities, Entities);
+
+		// Add sparse element to first two entities
+		FStructView _[] = {
+			EntityManager->AddSparseElementToEntity(Entities[0], FTestFragment_SparseInt::StaticStruct())
+			, EntityManager->AddSparseElementToEntity(Entities[1], FTestFragment_SparseInt::StaticStruct())
+		};
+
+		FMassArchetypeData& Archetype = FMassArchetypeHelper::ArchetypeDataFromHandleChecked(IntsArchetype);
+		const FMassArchetypeChunk& Chunk = Archetype.DebugGetChunk(0);
+		const UE::Mass::FChunkSparseElements& ChunkSparseElements = Chunk.GetSparseElementsUnsafe();
+
+		AITEST_TRUE("Chunk bitset should contain sparse int after adds",
+			ChunkSparseElements.GetSparseElementsPresent().Contains(FTestFragment_SparseInt::StaticStruct()));
+
+		// Remove from one entity - should still be in chunk bitset
+		EntityManager->RemoveSparseElementFromEntity(Entities[0], FTestFragment_SparseInt::StaticStruct());
+		AITEST_TRUE("Chunk bitset should still contain sparse int (one entity remains)",
+			ChunkSparseElements.GetSparseElementsPresent().Contains(FTestFragment_SparseInt::StaticStruct()));
+
+		// Remove from last entity - should be removed from chunk bitset
+		EntityManager->RemoveSparseElementFromEntity(Entities[1], FTestFragment_SparseInt::StaticStruct());
+		AITEST_FALSE("Chunk bitset should NOT contain sparse int after all removed",
+			ChunkSparseElements.GetSparseElementsPresent().Contains(FTestFragment_SparseInt::StaticStruct()));
+
+		return true;
+	}
+};
+IMPLEMENT_AI_INSTANT_TEST(FSparseElements_ChunkBitset_RemoveSingle, "System.Mass.Stress.Sparse.ChunkBitset.RemoveSingle");
+
+struct FSparseElements_ChunkBitset_BatchAdd : FEntityTestBase
+{
+	virtual bool InstantTest() override
+	{
+		CA_ASSUME(EntityManager);
+
+		constexpr int32 NumEntities = 50;
+		TArray<FMassEntityHandle> Entities;
+		EntityManager->BatchCreateEntities(IntsArchetype, NumEntities, Entities);
+
+		// Batch add sparse elements to first 20 entities
+		TArray<FMassEntityHandle> BatchEntities;
+		for (int32 i = 0; i < 20; ++i)
+		{
+			BatchEntities.Add(Entities[i]);
+		}
+		FMassArchetypeEntityCollection Collection = FMassArchetypeEntityCollection(IntsArchetype, BatchEntities, FMassArchetypeEntityCollection::NoDuplicates);
+		EntityManager->BatchAddSparseElementToEntities({Collection}, FTestFragment_SparseInt::StaticStruct());
+
+		// Verify chunk bitset is updated
+		FMassArchetypeData& Archetype = FMassArchetypeHelper::ArchetypeDataFromHandleChecked(IntsArchetype);
+		const FMassArchetypeChunk& Chunk = Archetype.DebugGetChunk(0);
+		const UE::Mass::FChunkSparseElements& ChunkSparseElements = Chunk.GetSparseElementsUnsafe();
+
+		AITEST_TRUE("Chunk bitset should contain sparse int after batch add",
+			ChunkSparseElements.GetSparseElementsPresent().Contains(FTestFragment_SparseInt::StaticStruct()));
+
+		return true;
+	}
+};
+IMPLEMENT_AI_INSTANT_TEST(FSparseElements_ChunkBitset_BatchAdd, "System.Mass.Stress.Sparse.ChunkBitset.BatchAdd");
+
+struct FSparseElements_ChunkBitset_BatchRemove : FEntityTestBase
+{
+	virtual bool InstantTest() override
+	{
+		CA_ASSUME(EntityManager);
+
+		constexpr int32 NumEntities = 50;
+		TArray<FMassEntityHandle> Entities;
+		EntityManager->BatchCreateEntities(IntsArchetype, NumEntities, Entities);
+
+		// Batch add to all entities
+		{
+			FMassArchetypeEntityCollection Collection = FMassArchetypeEntityCollection(IntsArchetype, Entities, FMassArchetypeEntityCollection::NoDuplicates);
+			EntityManager->BatchAddSparseElementToEntities({Collection}, FTestFragment_SparseInt::StaticStruct());
+		}
+
+		FMassArchetypeData& Archetype = FMassArchetypeHelper::ArchetypeDataFromHandleChecked(IntsArchetype);
+		const FMassArchetypeChunk& Chunk = Archetype.DebugGetChunk(0);
+		const UE::Mass::FChunkSparseElements& ChunkSparseElements = Chunk.GetSparseElementsUnsafe();
+
+		AITEST_TRUE("Chunk bitset should contain sparse int",
+			ChunkSparseElements.GetSparseElementsPresent().Contains(FTestFragment_SparseInt::StaticStruct()));
+
+		// Batch remove from half the entities - should still be in chunk bitset
+		TArray<FMassEntityHandle> HalfEntities;
+		for (int32 i = 0; i < 25; ++i)
+		{
+			HalfEntities.Add(Entities[i]);
+		}
+		{
+			FMassArchetypeEntityCollection Collection = FMassArchetypeEntityCollection(IntsArchetype, HalfEntities, FMassArchetypeEntityCollection::NoDuplicates);
+			EntityManager->BatchRemoveSparseElementFromEntities({Collection}, FTestFragment_SparseInt::StaticStruct());
+		}
+
+		AITEST_TRUE("Chunk bitset should still contain sparse int (some entities remain)",
+			ChunkSparseElements.GetSparseElementsPresent().Contains(FTestFragment_SparseInt::StaticStruct()));
+
+		// Batch remove from remaining entities - should be removed from chunk bitset
+		TArray<FMassEntityHandle> RemainingEntities;
+		for (int32 i = 25; i < NumEntities; ++i)
+		{
+			RemainingEntities.Add(Entities[i]);
+		}
+		{
+			FMassArchetypeEntityCollection Collection = FMassArchetypeEntityCollection(IntsArchetype, RemainingEntities, FMassArchetypeEntityCollection::NoDuplicates);
+			EntityManager->BatchRemoveSparseElementFromEntities({Collection}, FTestFragment_SparseInt::StaticStruct());
+		}
+
+		AITEST_FALSE("Chunk bitset should NOT contain sparse int after all removed",
+			ChunkSparseElements.GetSparseElementsPresent().Contains(FTestFragment_SparseInt::StaticStruct()));
+
+		return true;
+	}
+};
+IMPLEMENT_AI_INSTANT_TEST(FSparseElements_ChunkBitset_BatchRemove, "System.Mass.Stress.Sparse.ChunkBitset.BatchRemove");
+
+struct FSparseElements_ChunkBitset_MultipleSparseTypes : FEntityTestBase
+{
+	virtual bool InstantTest() override
+	{
+		CA_ASSUME(EntityManager);
+
+		constexpr int32 NumEntities = 20;
+		TArray<FMassEntityHandle> Entities;
+		EntityManager->BatchCreateEntities(IntsArchetype, NumEntities, Entities);
+
+		// Add different sparse types to different entities
+		FStructView _[] = {
+			EntityManager->AddSparseElementToEntity(Entities[0], FTestFragment_SparseInt::StaticStruct())
+			, EntityManager->AddSparseElementToEntity(Entities[1], FTestFragment_SparseFloat::StaticStruct())
+			, EntityManager->AddSparseElementToEntity(Entities[2], FTestTag_SparseA::StaticStruct())
+		};
+
+		// Verify chunk bitset contains all three types
+		FMassArchetypeData& Archetype = FMassArchetypeHelper::ArchetypeDataFromHandleChecked(IntsArchetype);
+		const FMassArchetypeChunk& Chunk = Archetype.DebugGetChunk(0);
+		const UE::Mass::FChunkSparseElements& ChunkSparseElements = Chunk.GetSparseElementsUnsafe();
+
+		AITEST_TRUE("Chunk bitset should contain sparse int",
+			ChunkSparseElements.GetSparseElementsPresent().Contains(FTestFragment_SparseInt::StaticStruct()));
+		AITEST_TRUE("Chunk bitset should contain sparse float",
+			ChunkSparseElements.GetSparseElementsPresent().Contains(FTestFragment_SparseFloat::StaticStruct()));
+		AITEST_TRUE("Chunk bitset should contain sparse tag A",
+			ChunkSparseElements.GetSparseElementsPresent().Contains(FTestTag_SparseA::StaticStruct()));
+
+		// Remove one type completely
+		EntityManager->RemoveSparseElementFromEntity(Entities[0], FTestFragment_SparseInt::StaticStruct());
+
+		AITEST_FALSE("Chunk bitset should NOT contain sparse int after removal",
+			ChunkSparseElements.GetSparseElementsPresent().Contains(FTestFragment_SparseInt::StaticStruct()));
+		AITEST_TRUE("Chunk bitset should still contain sparse float",
+			ChunkSparseElements.GetSparseElementsPresent().Contains(FTestFragment_SparseFloat::StaticStruct()));
+		AITEST_TRUE("Chunk bitset should still contain sparse tag A",
+			ChunkSparseElements.GetSparseElementsPresent().Contains(FTestTag_SparseA::StaticStruct()));
+
+		return true;
+	}
+};
+IMPLEMENT_AI_INSTANT_TEST(FSparseElements_ChunkBitset_MultipleSparseTypes, "System.Mass.Stress.Sparse.ChunkBitset.MultipleSparseTypes");
+
+struct FSparseElements_ChunkBitset_CrossChunk : FEntityTestBase
+{
+	virtual bool InstantTest() override
+	{
+		CA_ASSUME(EntityManager);
+
+		FMassArchetypeData& Archetype = FMassArchetypeHelper::ArchetypeDataFromHandleChecked(IntsArchetype);
+		const int32 NumEntitiesPerChunk = Archetype.GetNumEntitiesPerChunk();
+		int32 NumEntities = NumEntitiesPerChunk + 1;
+
+		TArray<FMassEntityHandle> Entities;
+		EntityManager->BatchCreateEntities(IntsArchetype, NumEntities, Entities);
+
+		// Add sparse element to entities in first chunk only
+		for (int32 i = 0; i < 64; ++i)
+		{
+			FStructView _ = EntityManager->AddSparseElementToEntity(Entities[i], FTestFragment_SparseInt::StaticStruct());
+		}
+
+		// Verify first chunk has the sparse element
+		const FMassArchetypeChunk& Chunk0 = Archetype.DebugGetChunk(0);
+		AITEST_TRUE("First chunk should have sparse elements", Chunk0.HasSparseElements());
+		AITEST_TRUE("First chunk bitset should contain sparse int",
+			Chunk0.GetSparseElementsUnsafe().GetSparseElementsPresent().Contains(FTestFragment_SparseInt::StaticStruct()));
+
+		// Verify second chunk does NOT have the sparse element
+		const FMassArchetypeChunk& Chunk1 = Archetype.DebugGetChunk(1);
+		AITEST_FALSE("Second chunk should not have sparse int in bitset",
+			Chunk1.HasSparseElements() && Chunk1.GetSparseElementsUnsafe().GetSparseElementsPresent().Contains(FTestFragment_SparseInt::StaticStruct()));
+
+		// Add sparse element to second chunk
+		{
+			FStructView _ = EntityManager->AddSparseElementToEntity(Entities[NumEntitiesPerChunk], FTestFragment_SparseInt::StaticStruct());
+		}
+		AITEST_TRUE("Second chunk should now have sparse elements", Chunk1.HasSparseElements());
+		AITEST_TRUE("Second chunk bitset should now contain sparse int",
+			Chunk1.GetSparseElementsUnsafe().GetSparseElementsPresent().Contains(FTestFragment_SparseInt::StaticStruct()));
+
+		return true;
+	}
+};
+IMPLEMENT_AI_INSTANT_TEST(FSparseElements_ChunkBitset_CrossChunk, "System.Mass.Stress.Sparse.ChunkBitset.CrossChunk");
+
+struct FSparseElements_ArchetypeBitset_SingleChunk : FEntityTestBase
+{
+	virtual bool InstantTest() override
+	{
+		CA_ASSUME(EntityManager);
+
+		constexpr int32 NumEntities = 10;
+		TArray<FMassEntityHandle> Entities;
+		EntityManager->BatchCreateEntities(IntsArchetype, NumEntities, Entities);
+
+		FMassArchetypeData& Archetype = FMassArchetypeHelper::ArchetypeDataFromHandleChecked(IntsArchetype);
+
+		// Initially, archetype should not contain sparse elements
+		AITEST_FALSE("Archetype should not contain sparse int initially",
+			Archetype.DoesContainEntitiesWithSparseElement(FTestFragment_SparseInt::StaticStruct()));
+
+		// Add sparse element
+		FStructView _ = EntityManager->AddSparseElementToEntity(Entities[0], FTestFragment_SparseInt::StaticStruct());
+
+		// Archetype bitset should now contain the sparse element
+		AITEST_TRUE("Archetype should contain sparse int after add",
+			Archetype.DoesContainEntitiesWithSparseElement(FTestFragment_SparseInt::StaticStruct()));
+
+		// Remove the sparse element from all entities
+		EntityManager->RemoveSparseElementFromEntity(Entities[0], FTestFragment_SparseInt::StaticStruct());
+
+		// Archetype bitset should no longer contain the sparse element
+		AITEST_FALSE("Archetype should not contain sparse int after removal",
+			Archetype.DoesContainEntitiesWithSparseElement(FTestFragment_SparseInt::StaticStruct()));
+
+		return true;
+	}
+};
+IMPLEMENT_AI_INSTANT_TEST(FSparseElements_ArchetypeBitset_SingleChunk, "System.Mass.Stress.Sparse.ArchetypeBitset.SingleChunk");
+
+struct FSparseElements_ArchetypeBitset_MultipleChunks : FEntityTestBase
+{
+	virtual bool InstantTest() override
+	{
+		CA_ASSUME(EntityManager);
+
+		const int32 EntitiesPerChunk = EntityManager->DebugGetArchetypeEntitiesCountPerChunk(IntsArchetype);
+		const int32 NumEntities = EntitiesPerChunk + 1;  // Multiple chunks
+		TArray<FMassEntityHandle> Entities;
+		EntityManager->BatchCreateEntities(IntsArchetype, NumEntities, Entities);
+
+		FMassArchetypeData& Archetype = FMassArchetypeHelper::ArchetypeDataFromHandleChecked(IntsArchetype);
+
+		// Add sparse element to entities in different chunks
+		FStructView _ [] = {
+			EntityManager->AddSparseElementToEntity(Entities[10], FTestFragment_SparseInt::StaticStruct())   // Chunk 0
+			, EntityManager->AddSparseElementToEntity(Entities[EntitiesPerChunk], FTestFragment_SparseFloat::StaticStruct()) // Chunk 1
+		};
+
+		// Archetype should contain both types
+		AITEST_TRUE("Archetype should contain sparse int",
+			Archetype.DoesContainEntitiesWithSparseElement(FTestFragment_SparseInt::StaticStruct()));
+		AITEST_TRUE("Archetype should contain sparse float",
+			Archetype.DoesContainEntitiesWithSparseElement(FTestFragment_SparseFloat::StaticStruct()));
+
+		// Remove sparse int from first chunk
+		EntityManager->RemoveSparseElementFromEntity(Entities[10], FTestFragment_SparseInt::StaticStruct());
+
+		AITEST_FALSE("Archetype should not contain sparse int after removal",
+			Archetype.DoesContainEntitiesWithSparseElement(FTestFragment_SparseInt::StaticStruct()));
+		AITEST_TRUE("Archetype should still contain sparse float",
+			Archetype.DoesContainEntitiesWithSparseElement(FTestFragment_SparseFloat::StaticStruct()));
+
+		return true;
+	}
+};
+IMPLEMENT_AI_INSTANT_TEST(FSparseElements_ArchetypeBitset_MultipleChunks, "System.Mass.Stress.Sparse.ArchetypeBitset.MultipleChunks");
+
+struct FSparseElements_EntityRemoval_ChunkBitsetUpdate : FEntityTestBase
+{
+	virtual bool InstantTest() override
+	{
+		CA_ASSUME(EntityManager);
+
+		constexpr int32 NumEntities = 20;
+		TArray<FMassEntityHandle> Entities;
+		EntityManager->BatchCreateEntities(IntsArchetype, NumEntities, Entities);
+
+		// Add sparse elements to first 5 entities
+		for (int32 i = 0; i < 5; ++i)
+		{
+			FStructView _ = EntityManager->AddSparseElementToEntity(Entities[i], FTestFragment_SparseInt::StaticStruct());
+		}
+
+		FMassArchetypeData& Archetype = FMassArchetypeHelper::ArchetypeDataFromHandleChecked(IntsArchetype);
+		const FMassArchetypeChunk& Chunk = Archetype.DebugGetChunk(0);
+		const UE::Mass::FChunkSparseElements& ChunkSparseElements = Chunk.GetSparseElementsUnsafe();
+
+		AITEST_TRUE("Chunk should contain sparse int",
+			ChunkSparseElements.GetSparseElementsPresent().Contains(FTestFragment_SparseInt::StaticStruct()));
+
+		// Destroy entities with sparse elements one by one (but not all)
+		for (int32 i = 0; i < 3; ++i)
+		{
+			EntityManager->DestroyEntity(Entities[i]);
+		}
+		EntityManager->FlushCommands();
+
+		// Chunk should still contain sparse int (2 entities remain)
+		AITEST_TRUE("Chunk should still contain sparse int after partial removal",
+			ChunkSparseElements.GetSparseElementsPresent().Contains(FTestFragment_SparseInt::StaticStruct()));
+
+		// Destroy remaining entities with sparse elements
+		for (int32 i = 3; i < 5; ++i)
+		{
+			EntityManager->DestroyEntity(Entities[i]);
+		}
+		EntityManager->FlushCommands();
+
+		// Chunk should no longer contain sparse int
+		AITEST_FALSE("Chunk should not contain sparse int after all removed",
+			ChunkSparseElements.GetSparseElementsPresent().Contains(FTestFragment_SparseInt::StaticStruct()));
+
+		return true;
+	}
+};
+IMPLEMENT_AI_INSTANT_TEST(FSparseElements_EntityRemoval_ChunkBitsetUpdate, "System.Mass.Stress.Sparse.EntityRemoval.ChunkBitsetUpdate");
+
+struct FSparseElements_BatchEntityDestruction_ChunkBitsetUpdate : FEntityTestBase
+{
+	virtual bool InstantTest() override
+	{
+		CA_ASSUME(EntityManager);
+
+		constexpr int32 NumEntities = 50;
+		TArray<FMassEntityHandle> Entities;
+		EntityManager->BatchCreateEntities(IntsArchetype, NumEntities, Entities);
+
+		// Add sparse elements to all entities
+		{
+			FMassArchetypeEntityCollection Collection = FMassArchetypeEntityCollection(IntsArchetype, Entities, FMassArchetypeEntityCollection::NoDuplicates);
+			EntityManager->BatchAddSparseElementToEntities({Collection}, FTestFragment_SparseInt::StaticStruct());
+		}
+
+		FMassArchetypeData& Archetype = FMassArchetypeHelper::ArchetypeDataFromHandleChecked(IntsArchetype);
+		const FMassArchetypeChunk& Chunk = Archetype.DebugGetChunk(0);
+		const UE::Mass::FChunkSparseElements& ChunkSparseElements = Chunk.GetSparseElementsUnsafe();
+
+		AITEST_TRUE("Chunk should contain sparse int",
+			ChunkSparseElements.GetSparseElementsPresent().Contains(FTestFragment_SparseInt::StaticStruct()));
+
+		// Batch destroy all entities
+		EntityManager->BatchDestroyEntities(Entities);
+		EntityManager->FlushCommands();
+
+		// Chunk should no longer exist or be empty
+		AITEST_EQUAL("Archetype should have no entities after batch destruction", Archetype.GetNumEntities(), 0);
+
+		return true;
+	}
+};
+IMPLEMENT_AI_INSTANT_TEST(FSparseElements_BatchEntityDestruction_ChunkBitsetUpdate, "System.Mass.Stress.Sparse.EntityRemoval.BatchDestruction");
+
+struct FSparseElements_ArchetypeChange_BitsetTracking : FEntityTestBase
+{
+	virtual bool InstantTest() override
+	{
+		CA_ASSUME(EntityManager);
+
+		// Create source archetype with just Int
+		FMassArchetypeCompositionDescriptor SourceComposition;
+		SourceComposition.Add<FTestFragment_Int>();
+		const FMassArchetypeHandle SourceArchetype = EntityManager->CreateArchetype(SourceComposition);
+
+		// Create target archetype with Int + Float
+		FMassArchetypeCompositionDescriptor TargetComposition;
+		TargetComposition.Add<FTestFragment_Int>();
+		TargetComposition.Add<FTestFragment_Float>();
+		const FMassArchetypeHandle TargetArchetype = EntityManager->CreateArchetype(TargetComposition);
+
+		constexpr int32 NumEntities = 20;
+		TArray<FMassEntityHandle> Entities;
+		EntityManager->BatchCreateEntities(SourceArchetype, NumEntities, Entities);
+
+		// Add sparse elements to entities in source archetype
+		for (int32 i = 0; i < 10; ++i)
+		{
+			FStructView _ = EntityManager->AddSparseElementToEntity(Entities[i], FTestFragment_SparseInt::StaticStruct());
+		}
+
+		FMassArchetypeData& SourceArchetypeData = FMassArchetypeHelper::ArchetypeDataFromHandleChecked(SourceArchetype);
+		AITEST_TRUE("Source archetype should contain sparse int",
+			SourceArchetypeData.DoesContainEntitiesWithSparseElement(FTestFragment_SparseInt::StaticStruct()));
+
+		// Move entities to target archetype
+		for (int32 i = 0; i < 10; ++i)
+		{
+			EntityManager->AddFragmentToEntity(Entities[i], FTestFragment_Float::StaticStruct());
+		}
+
+		FMassArchetypeData& TargetArchetypeData = FMassArchetypeHelper::ArchetypeDataFromHandleChecked(TargetArchetype);
+
+		// Target archetype should now contain sparse elements
+		AITEST_TRUE("Target archetype should contain sparse int after move",
+			TargetArchetypeData.DoesContainEntitiesWithSparseElement(FTestFragment_SparseInt::StaticStruct()));
+
+		// Source archetype should still have entities with sparse elements (remaining 10 entities)
+		AITEST_TRUE("Source archetype should no longer contain sparse int - all sparse fragment owning entities have been moved",
+			SourceArchetypeData.DoesContainEntitiesWithSparseElement(FTestFragment_SparseInt::StaticStruct()) == false);
+
+		return true;
+	}
+};
+IMPLEMENT_AI_INSTANT_TEST(FSparseElements_ArchetypeChange_BitsetTracking, "System.Mass.Stress.Sparse.ArchetypeChange.BitsetTracking");
+
+struct FSparseElements_ComplexScenario_BitsetAccuracy : FEntityTestBase
+{
+	virtual bool InstantTest() override
+	{
+		CA_ASSUME(EntityManager);
+
+		constexpr int32 NumEntities = 100;
+		TArray<FMassEntityHandle> Entities;
+		EntityManager->BatchCreateEntities(IntsArchetype, NumEntities, Entities);
+
+		FMassArchetypeData& Archetype = FMassArchetypeHelper::ArchetypeDataFromHandleChecked(IntsArchetype);
+
+		// Complex sequence of operations
+		// 1. Batch add sparse int to first 50
+		TArray<FMassEntityHandle> FirstHalf;
+		for (int32 i = 0; i < 50; ++i)
+		{
+			FirstHalf.Add(Entities[i]);
+		}
+		{
+			FMassArchetypeEntityCollection Collection = FMassArchetypeEntityCollection(IntsArchetype, FirstHalf, FMassArchetypeEntityCollection::NoDuplicates);
+			EntityManager->BatchAddSparseElementToEntities({Collection}, FTestFragment_SparseInt::StaticStruct());
+		}
+		AITEST_TRUE("Should contain sparse int after batch add",
+			Archetype.DoesContainEntitiesWithSparseElement(FTestFragment_SparseInt::StaticStruct()));
+
+		// 2. Add sparse float to entities 25-75 (overlapping)
+		for (int32 i = 25; i < 75; ++i)
+		{
+			FStructView _ = EntityManager->AddSparseElementToEntity(Entities[i], FTestFragment_SparseFloat::StaticStruct());
+		}
+		AITEST_TRUE("Should contain sparse float",
+			Archetype.DoesContainEntitiesWithSparseElement(FTestFragment_SparseFloat::StaticStruct()));
+
+		// 3. Remove sparse int from 0-25
+		TArray<FMassEntityHandle> FirstQuarter;
+		for (int32 i = 0; i < 25; ++i)
+		{
+			FirstQuarter.Add(Entities[i]);
+		}
+		{
+			FMassArchetypeEntityCollection Collection = FMassArchetypeEntityCollection(IntsArchetype, FirstQuarter, FMassArchetypeEntityCollection::NoDuplicates);
+			EntityManager->BatchRemoveSparseElementFromEntities({Collection}, FTestFragment_SparseInt::StaticStruct());
+		}
+		AITEST_TRUE("Should still contain sparse int (entities 25-49 remain)",
+			Archetype.DoesContainEntitiesWithSparseElement(FTestFragment_SparseInt::StaticStruct()));
+
+		// 4. Destroy entities 40-60
+		for (int32 i = 40; i < 60; ++i)
+		{
+			EntityManager->DestroyEntity(Entities[i]);
+		}
+		EntityManager->FlushCommands();
+
+		// Both types should still be present
+		AITEST_TRUE("Should still contain sparse int (entities 25-39 remain)",
+			Archetype.DoesContainEntitiesWithSparseElement(FTestFragment_SparseInt::StaticStruct()));
+		AITEST_TRUE("Should still contain sparse float (entities 25-39 and 60-74 remain)",
+			Archetype.DoesContainEntitiesWithSparseElement(FTestFragment_SparseFloat::StaticStruct()));
+
+		// 5. Remove all remaining sparse int
+		TArray<FMassEntityHandle> RemainingWithInt;
+		for (int32 i = 25; i < 40; ++i)
+		{
+			RemainingWithInt.Add(Entities[i]);
+		}
+		{
+			FMassArchetypeEntityCollection Collection = FMassArchetypeEntityCollection(IntsArchetype, RemainingWithInt, FMassArchetypeEntityCollection::NoDuplicates);
+			EntityManager->BatchRemoveSparseElementFromEntities({Collection}, FTestFragment_SparseInt::StaticStruct());
+		}
+
+		AITEST_FALSE("Should NOT contain sparse int after all removed",
+			Archetype.DoesContainEntitiesWithSparseElement(FTestFragment_SparseInt::StaticStruct()));
+		AITEST_TRUE("Should still contain sparse float",
+			Archetype.DoesContainEntitiesWithSparseElement(FTestFragment_SparseFloat::StaticStruct()));
+
+		return true;
+	}
+};
+IMPLEMENT_AI_INSTANT_TEST(FSparseElements_ComplexScenario_BitsetAccuracy, "System.Mass.Stress.Sparse.Complex.BitsetAccuracy");
+
+struct FSparseElements_ReAdd_ChunkBitsetConsistency : FEntityTestBase
+{
+	virtual bool InstantTest() override
+	{
+		CA_ASSUME(EntityManager);
+
+		constexpr int32 NumEntities = 10;
+		TArray<FMassEntityHandle> Entities;
+		EntityManager->BatchCreateEntities(IntsArchetype, NumEntities, Entities);
+
+		FMassArchetypeData& Archetype = FMassArchetypeHelper::ArchetypeDataFromHandleChecked(IntsArchetype);
+		const FMassArchetypeChunk& Chunk = Archetype.DebugGetChunk(0);
+		const UE::Mass::FChunkSparseElements& ChunkSparseElements = Chunk.GetSparseElementsUnsafe();
+
+		// Add -> Remove -> Add cycle
+		{
+			FStructView _ = EntityManager->AddSparseElementToEntity(Entities[0], FTestFragment_SparseInt::StaticStruct());
+		}
+		AITEST_TRUE("Chunk should contain sparse int after first add",
+			ChunkSparseElements.GetSparseElementsPresent().Contains(FTestFragment_SparseInt::StaticStruct()));
+
+		EntityManager->RemoveSparseElementFromEntity(Entities[0], FTestFragment_SparseInt::StaticStruct());
+		AITEST_FALSE("Chunk should NOT contain sparse int after remove",
+			ChunkSparseElements.GetSparseElementsPresent().Contains(FTestFragment_SparseInt::StaticStruct()));
+
+		{
+			FStructView _ = EntityManager->AddSparseElementToEntity(Entities[0], FTestFragment_SparseInt::StaticStruct());
+		}
+		AITEST_TRUE("Chunk should contain sparse int after re-add",
+			ChunkSparseElements.GetSparseElementsPresent().Contains(FTestFragment_SparseInt::StaticStruct()));
+
+		// Add to second entity and remove from first - bitset should remain
+		{
+			FStructView _ = EntityManager->AddSparseElementToEntity(Entities[1], FTestFragment_SparseInt::StaticStruct());
+		}
+		EntityManager->RemoveSparseElementFromEntity(Entities[0], FTestFragment_SparseInt::StaticStruct());
+		AITEST_TRUE("Chunk should still contain sparse int (second entity has it)",
+			ChunkSparseElements.GetSparseElementsPresent().Contains(FTestFragment_SparseInt::StaticStruct()));
+
+		return true;
+	}
+};
+IMPLEMENT_AI_INSTANT_TEST(FSparseElements_ReAdd_ChunkBitsetConsistency, "System.Mass.Stress.Sparse.ReAdd.ChunkBitsetConsistency");
+
+
+//----------------------------------------------------------------------//
+// Unified API Tests - verifying Fragment/Tag API works seamlessly with Sparse Elements
+//----------------------------------------------------------------------//
+
+struct FSparseElements_UnifiedAPI_AddFragmentToEntity : FEntityTestBase
+{
+	virtual bool InstantTest() override
+	{
+		const FMassEntityHandle Entity = EntityManager->CreateEntity(IntsArchetype);
+
+		// Verify entity doesn't have sparse fragment initially
+		AITEST_FALSE("Entity should not have sparse fragment initially",
+			EntityManager->DoesEntityHaveElement(Entity, FTestFragment_SparseInt::StaticStruct()));
+
+		// Use AddFragmentToEntity with sparse fragment type
+		EntityManager->AddFragmentToEntity(Entity, FTestFragment_SparseInt::StaticStruct());
+
+		// Verify sparse fragment was added
+		AITEST_TRUE("Entity should have sparse fragment after AddFragmentToEntity",
+			EntityManager->DoesEntityHaveElement(Entity, FTestFragment_SparseInt::StaticStruct()));
+
+		// Verify entity remained in same archetype (sparse doesn't change archetype)
+		AITEST_EQUAL("Entity should remain in original archetype",
+			EntityManager->GetArchetypeForEntity(Entity), IntsArchetype);
+
+		// Verify we can access the sparse fragment data
+		FMassEntityView EntityView(*EntityManager, Entity);
+		const FTestFragment_SparseInt* SparseFragment = EntityView.GetSparseFragmentDataPtr<FTestFragment_SparseInt>();
+		AITEST_NOT_NULL("Should be able to access sparse fragment data", SparseFragment);
+
+		return true;
+	}
+};
+IMPLEMENT_AI_INSTANT_TEST(FSparseElements_UnifiedAPI_AddFragmentToEntity, "System.Mass.Stress.Sparse.UnifiedAPI.AddFragmentToEntity");
+
+
+struct FSparseElements_UnifiedAPI_AddFragmentToEntity_WithInitializer : FEntityTestBase
+{
+	virtual bool InstantTest() override
+	{
+		constexpr int32 InitialValue = 42;
+		const FMassEntityHandle Entity = EntityManager->CreateEntity(IntsArchetype);
+
+		// Use AddFragmentToEntity with initializer callback for sparse fragment
+		EntityManager->AddFragmentToEntity(Entity, FTestFragment_SparseInt::StaticStruct(),
+			[InitialValue](void* FragmentMemory, const UScriptStruct& FragmentType)
+			{
+				FTestFragment_SparseInt* Fragment = static_cast<FTestFragment_SparseInt*>(FragmentMemory);
+				Fragment->Value = InitialValue;
+			});
+
+		// Verify sparse fragment was added with correct value
+		AITEST_TRUE("Entity should have sparse fragment",
+			EntityManager->DoesEntityHaveElement(Entity, FTestFragment_SparseInt::StaticStruct()));
+
+		FMassEntityView EntityView(*EntityManager, Entity);
+		const FTestFragment_SparseInt* SparseFragment = EntityView.GetSparseFragmentDataPtr<FTestFragment_SparseInt>();
+		AITEST_NOT_NULL("Should be able to access sparse fragment data", SparseFragment);
+		AITEST_EQUAL("Sparse fragment should have initialized value", SparseFragment->Value, InitialValue);
+
+		return true;
+	}
+};
+IMPLEMENT_AI_INSTANT_TEST(FSparseElements_UnifiedAPI_AddFragmentToEntity_WithInitializer, "System.Mass.Stress.Sparse.UnifiedAPI.AddFragmentToEntity.WithInitializer");
+
+
+struct FSparseElements_UnifiedAPI_RemoveFragmentFromEntity : FEntityTestBase
+{
+	virtual bool InstantTest() override
+	{
+		const FMassEntityHandle Entity = EntityManager->CreateEntity(IntsArchetype);
+
+		// Add sparse fragment first
+		EntityManager->AddFragmentToEntity(Entity, FTestFragment_SparseInt::StaticStruct());
+		AITEST_TRUE("Entity should have sparse fragment after add",
+			EntityManager->DoesEntityHaveElement(Entity, FTestFragment_SparseInt::StaticStruct()));
+
+		// Use RemoveFragmentFromEntity with sparse fragment type
+		EntityManager->RemoveFragmentFromEntity(Entity, FTestFragment_SparseInt::StaticStruct());
+
+		// Verify sparse fragment was removed
+		AITEST_FALSE("Entity should not have sparse fragment after RemoveFragmentFromEntity",
+			EntityManager->DoesEntityHaveElement(Entity, FTestFragment_SparseInt::StaticStruct()));
+
+		// Verify entity remained in same archetype
+		AITEST_EQUAL("Entity should remain in original archetype",
+			EntityManager->GetArchetypeForEntity(Entity), IntsArchetype);
+
+		return true;
+	}
+};
+IMPLEMENT_AI_INSTANT_TEST(FSparseElements_UnifiedAPI_RemoveFragmentFromEntity, "System.Mass.Stress.Sparse.UnifiedAPI.RemoveFragmentFromEntity");
+
+
+struct FSparseElements_UnifiedAPI_AddTagToEntity : FEntityTestBase
+{
+	virtual bool InstantTest() override
+	{
+		const FMassEntityHandle Entity = EntityManager->CreateEntity(IntsArchetype);
+
+		// Verify entity doesn't have sparse tag initially
+		AITEST_FALSE("Entity should not have sparse tag initially",
+			EntityManager->DoesEntityHaveElement(Entity, FTestTag_SparseA::StaticStruct()));
+
+		// Use AddTagToEntity with sparse tag type
+		EntityManager->AddTagToEntity(Entity, FTestTag_SparseA::StaticStruct());
+
+		// Verify sparse tag was added
+		AITEST_TRUE("Entity should have sparse tag after AddTagToEntity",
+			EntityManager->DoesEntityHaveElement(Entity, FTestTag_SparseA::StaticStruct()));
+
+		// Verify entity remained in same archetype (sparse doesn't change archetype)
+		AITEST_EQUAL("Entity should remain in original archetype",
+			EntityManager->GetArchetypeForEntity(Entity), IntsArchetype);
+
+		// Verify via FMassEntityView
+		FMassEntityView EntityView(*EntityManager, Entity);
+		AITEST_TRUE("EntityView should report having sparse tag", EntityView.HasElement<FTestTag_SparseA>());
+
+		return true;
+	}
+};
+IMPLEMENT_AI_INSTANT_TEST(FSparseElements_UnifiedAPI_AddTagToEntity, "System.Mass.Stress.Sparse.UnifiedAPI.AddTagToEntity");
+
+
+struct FSparseElements_UnifiedAPI_RemoveTagFromEntity : FEntityTestBase
+{
+	virtual bool InstantTest() override
+	{
+		const FMassEntityHandle Entity = EntityManager->CreateEntity(IntsArchetype);
+
+		// Add sparse tag first
+		EntityManager->AddTagToEntity(Entity, FTestTag_SparseA::StaticStruct());
+		AITEST_TRUE("Entity should have sparse tag after add",
+			EntityManager->DoesEntityHaveElement(Entity, FTestTag_SparseA::StaticStruct()));
+
+		// Use RemoveTagFromEntity with sparse tag type
+		EntityManager->RemoveTagFromEntity(Entity, FTestTag_SparseA::StaticStruct());
+
+		// Verify sparse tag was removed
+		AITEST_FALSE("Entity should not have sparse tag after RemoveTagFromEntity",
+			EntityManager->DoesEntityHaveElement(Entity, FTestTag_SparseA::StaticStruct()));
+
+		// Verify entity remained in same archetype
+		AITEST_EQUAL("Entity should remain in original archetype",
+			EntityManager->GetArchetypeForEntity(Entity), IntsArchetype);
+
+		return true;
+	}
+};
+IMPLEMENT_AI_INSTANT_TEST(FSparseElements_UnifiedAPI_RemoveTagFromEntity, "System.Mass.Stress.Sparse.UnifiedAPI.RemoveTagFromEntity");
+
+
+struct FSparseElements_UnifiedAPI_AddElementToEntity_SparseFragment : FEntityTestBase
+{
+	virtual bool InstantTest() override
+	{
+		const FMassEntityHandle Entity = EntityManager->CreateEntity(IntsArchetype);
+
+		// Use generic AddElementToEntity with sparse fragment type
+		EntityManager->AddElementToEntity(Entity, FTestFragment_SparseInt::StaticStruct());
+
+		// Verify sparse fragment was added
+		AITEST_TRUE("Entity should have sparse fragment after AddElementToEntity",
+			EntityManager->DoesEntityHaveElement(Entity, FTestFragment_SparseInt::StaticStruct()));
+
+		// Verify archetype unchanged
+		AITEST_EQUAL("Entity should remain in original archetype",
+			EntityManager->GetArchetypeForEntity(Entity), IntsArchetype);
+
+		return true;
+	}
+};
+IMPLEMENT_AI_INSTANT_TEST(FSparseElements_UnifiedAPI_AddElementToEntity_SparseFragment, "System.Mass.Stress.Sparse.UnifiedAPI.AddElementToEntity.SparseFragment");
+
+
+struct FSparseElements_UnifiedAPI_AddElementToEntity_SparseTag : FEntityTestBase
+{
+	virtual bool InstantTest() override
+	{
+		const FMassEntityHandle Entity = EntityManager->CreateEntity(IntsArchetype);
+
+		// Use generic AddElementToEntity with sparse tag type
+		EntityManager->AddElementToEntity(Entity, FTestTag_SparseA::StaticStruct());
+
+		// Verify sparse tag was added
+		AITEST_TRUE("Entity should have sparse tag after AddElementToEntity",
+			EntityManager->DoesEntityHaveElement(Entity, FTestTag_SparseA::StaticStruct()));
+
+		// Verify archetype unchanged
+		AITEST_EQUAL("Entity should remain in original archetype",
+			EntityManager->GetArchetypeForEntity(Entity), IntsArchetype);
+
+		return true;
+	}
+};
+IMPLEMENT_AI_INSTANT_TEST(FSparseElements_UnifiedAPI_AddElementToEntity_SparseTag, "System.Mass.Stress.Sparse.UnifiedAPI.AddElementToEntity.SparseTag");
+
+
+struct FSparseElements_UnifiedAPI_RemoveElementFromEntity_SparseFragment : FEntityTestBase
+{
+	virtual bool InstantTest() override
+	{
+		const FMassEntityHandle Entity = EntityManager->CreateEntity(IntsArchetype);
+
+		// Add then remove sparse fragment using generic API
+		EntityManager->AddElementToEntity(Entity, FTestFragment_SparseInt::StaticStruct());
+		AITEST_TRUE("Entity should have sparse fragment",
+			EntityManager->DoesEntityHaveElement(Entity, FTestFragment_SparseInt::StaticStruct()));
+
+		EntityManager->RemoveElementFromEntity(Entity, FTestFragment_SparseInt::StaticStruct());
+
+		// Verify sparse fragment was removed
+		AITEST_FALSE("Entity should not have sparse fragment after RemoveElementFromEntity",
+			EntityManager->DoesEntityHaveElement(Entity, FTestFragment_SparseInt::StaticStruct()));
+
+		return true;
+	}
+};
+IMPLEMENT_AI_INSTANT_TEST(FSparseElements_UnifiedAPI_RemoveElementFromEntity_SparseFragment, "System.Mass.Stress.Sparse.UnifiedAPI.RemoveElementFromEntity.SparseFragment");
+
+
+struct FSparseElements_UnifiedAPI_RemoveElementFromEntity_SparseTag : FEntityTestBase
+{
+	virtual bool InstantTest() override
+	{
+		const FMassEntityHandle Entity = EntityManager->CreateEntity(IntsArchetype);
+
+		// Add then remove sparse tag using generic API
+		EntityManager->AddElementToEntity(Entity, FTestTag_SparseA::StaticStruct());
+		AITEST_TRUE("Entity should have sparse tag",
+			EntityManager->DoesEntityHaveElement(Entity, FTestTag_SparseA::StaticStruct()));
+
+		EntityManager->RemoveElementFromEntity(Entity, FTestTag_SparseA::StaticStruct());
+
+		// Verify sparse tag was removed
+		AITEST_FALSE("Entity should not have sparse tag after RemoveElementFromEntity",
+			EntityManager->DoesEntityHaveElement(Entity, FTestTag_SparseA::StaticStruct()));
+
+		return true;
+	}
+};
+IMPLEMENT_AI_INSTANT_TEST(FSparseElements_UnifiedAPI_RemoveElementFromEntity_SparseTag, "System.Mass.Stress.Sparse.UnifiedAPI.RemoveElementFromEntity.SparseTag");
+
+
+struct FSparseElements_UnifiedAPI_DoesEntityHaveElement_AllSparseTypes : FEntityTestBase
+{
+	virtual bool InstantTest() override
+	{
+		const FMassEntityHandle Entity = EntityManager->CreateEntity(IntsArchetype);
+
+		// Verify none present initially
+		AITEST_FALSE("Should not have sparse int initially",
+			EntityManager->DoesEntityHaveElement(Entity, FTestFragment_SparseInt::StaticStruct()));
+		AITEST_FALSE("Should not have sparse float initially",
+			EntityManager->DoesEntityHaveElement(Entity, FTestFragment_SparseFloat::StaticStruct()));
+		AITEST_FALSE("Should not have sparse tag A initially",
+			EntityManager->DoesEntityHaveElement(Entity, FTestTag_SparseA::StaticStruct()));
+		AITEST_FALSE("Should not have sparse tag B initially",
+			EntityManager->DoesEntityHaveElement(Entity, FTestTag_SparseB::StaticStruct()));
+
+		// Add all sparse element types
+		EntityManager->AddElementToEntity(Entity, FTestFragment_SparseInt::StaticStruct());
+		EntityManager->AddElementToEntity(Entity, FTestFragment_SparseFloat::StaticStruct());
+		EntityManager->AddElementToEntity(Entity, FTestTag_SparseA::StaticStruct());
+		EntityManager->AddElementToEntity(Entity, FTestTag_SparseB::StaticStruct());
+
+		// Verify all present
+		AITEST_TRUE("Should have sparse int",
+			EntityManager->DoesEntityHaveElement(Entity, FTestFragment_SparseInt::StaticStruct()));
+		AITEST_TRUE("Should have sparse float",
+			EntityManager->DoesEntityHaveElement(Entity, FTestFragment_SparseFloat::StaticStruct()));
+		AITEST_TRUE("Should have sparse tag A",
+			EntityManager->DoesEntityHaveElement(Entity, FTestTag_SparseA::StaticStruct()));
+		AITEST_TRUE("Should have sparse tag B",
+			EntityManager->DoesEntityHaveElement(Entity, FTestTag_SparseB::StaticStruct()));
+
+		// Also verify template version
+		AITEST_TRUE("Template DoesEntityHaveElement should work for sparse int",
+			EntityManager->DoesEntityHaveElement<FTestFragment_SparseInt>(Entity));
+		AITEST_TRUE("Template DoesEntityHaveElement should work for sparse tag",
+			EntityManager->DoesEntityHaveElement<FTestTag_SparseA>(Entity));
+
+		return true;
+	}
+};
+IMPLEMENT_AI_INSTANT_TEST(FSparseElements_UnifiedAPI_DoesEntityHaveElement_AllSparseTypes, "System.Mass.Stress.Sparse.UnifiedAPI.DoesEntityHaveElement.AllTypes");
+
+
+struct FSparseElements_UnifiedAPI_EntityView_HasElement : FEntityTestBase
+{
+	virtual bool InstantTest() override
+	{
+		const FMassEntityHandle Entity = EntityManager->CreateEntity(IntsArchetype);
+		FMassEntityView EntityView(*EntityManager, Entity);
+
+		// Verify HasElement returns false initially for sparse elements
+		AITEST_FALSE("EntityView should not report sparse fragment initially",
+			EntityView.HasElement<FTestFragment_SparseInt>());
+		AITEST_FALSE("EntityView should not report sparse tag initially",
+			EntityView.HasElement<FTestTag_SparseA>());
+
+		// Add sparse elements
+		EntityManager->AddElementToEntity(Entity, FTestFragment_SparseInt::StaticStruct());
+		EntityManager->AddElementToEntity(Entity, FTestTag_SparseA::StaticStruct());
+
+		// Verify HasElement returns true
+		AITEST_TRUE("EntityView should report sparse fragment after add",
+			EntityView.HasElement<FTestFragment_SparseInt>());
+		AITEST_TRUE("EntityView should report sparse tag after add",
+			EntityView.HasElement<FTestTag_SparseA>());
+
+		// Test with EIncludeSparseElements parameter
+		AITEST_TRUE("HasElement with IncludeSparseElements::Yes should return true",
+			EntityView.HasElement<FTestFragment_SparseInt>(UE::Mass::EIncludeSparseElements::Yes));
+		AITEST_FALSE("HasElement with IncludeSparseElements::No should return false for sparse element",
+			EntityView.HasElement<FTestFragment_SparseInt>(UE::Mass::EIncludeSparseElements::No));
+
+		return true;
+	}
+};
+IMPLEMENT_AI_INSTANT_TEST(FSparseElements_UnifiedAPI_EntityView_HasElement, "System.Mass.Stress.Sparse.UnifiedAPI.EntityView.HasElement");
+
+
+struct FSparseElements_UnifiedAPI_MixedRegularAndSparse : FEntityTestBase
+{
+	virtual bool InstantTest() override
+	{
+		const FMassEntityHandle Entity = EntityManager->CreateEntity(IntsArchetype);
+
+		// Add regular fragment (changes archetype)
+		EntityManager->AddFragmentToEntity(Entity, FTestFragment_Float::StaticStruct());
+		const FMassArchetypeHandle NewArchetype = EntityManager->GetArchetypeForEntity(Entity);
+		AITEST_NOT_EQUAL("Archetype should change after adding regular fragment", NewArchetype, IntsArchetype);
+
+		// Add sparse fragment (should not change archetype)
+		EntityManager->AddFragmentToEntity(Entity, FTestFragment_SparseInt::StaticStruct());
+		AITEST_EQUAL("Archetype should not change after adding sparse fragment",
+			EntityManager->GetArchetypeForEntity(Entity), NewArchetype);
+
+		// Add regular tag (changes archetype)
+		EntityManager->AddTagToEntity(Entity, FTestTag_A::StaticStruct());
+		const FMassArchetypeHandle ArchetypeWithTag = EntityManager->GetArchetypeForEntity(Entity);
+		AITEST_NOT_EQUAL("Archetype should change after adding regular tag", ArchetypeWithTag, NewArchetype);
+
+		// Add sparse tag (should not change archetype)
+		EntityManager->AddTagToEntity(Entity, FTestTag_SparseA::StaticStruct());
+		AITEST_EQUAL("Archetype should not change after adding sparse tag",
+			EntityManager->GetArchetypeForEntity(Entity), ArchetypeWithTag);
+
+		// Verify all elements present
+		AITEST_TRUE("Should have regular int fragment (from original archetype)",
+			EntityManager->DoesEntityHaveElement(Entity, FTestFragment_Int::StaticStruct()));
+		AITEST_TRUE("Should have regular float fragment",
+			EntityManager->DoesEntityHaveElement(Entity, FTestFragment_Float::StaticStruct()));
+		AITEST_TRUE("Should have sparse int fragment",
+			EntityManager->DoesEntityHaveElement(Entity, FTestFragment_SparseInt::StaticStruct()));
+		AITEST_TRUE("Should have regular tag A",
+			EntityManager->DoesEntityHaveElement(Entity, FTestTag_A::StaticStruct()));
+		AITEST_TRUE("Should have sparse tag A",
+			EntityManager->DoesEntityHaveElement(Entity, FTestTag_SparseA::StaticStruct()));
+
+		// Remove sparse elements - archetype should not change
+		EntityManager->RemoveFragmentFromEntity(Entity, FTestFragment_SparseInt::StaticStruct());
+		EntityManager->RemoveTagFromEntity(Entity, FTestTag_SparseA::StaticStruct());
+		AITEST_EQUAL("Archetype should not change after removing sparse elements",
+			EntityManager->GetArchetypeForEntity(Entity), ArchetypeWithTag);
+
+		// Verify sparse elements removed
+		AITEST_FALSE("Should not have sparse int fragment after removal",
+			EntityManager->DoesEntityHaveElement(Entity, FTestFragment_SparseInt::StaticStruct()));
+		AITEST_FALSE("Should not have sparse tag A after removal",
+			EntityManager->DoesEntityHaveElement(Entity, FTestTag_SparseA::StaticStruct()));
+
+		// Verify regular elements still present
+		AITEST_TRUE("Should still have regular float fragment",
+			EntityManager->DoesEntityHaveElement(Entity, FTestFragment_Float::StaticStruct()));
+		AITEST_TRUE("Should still have regular tag A",
+			EntityManager->DoesEntityHaveElement(Entity, FTestTag_A::StaticStruct()));
+
+		return true;
+	}
+};
+IMPLEMENT_AI_INSTANT_TEST(FSparseElements_UnifiedAPI_MixedRegularAndSparse, "System.Mass.Stress.Sparse.UnifiedAPI.MixedRegularAndSparse");
+
+
+struct FSparseElements_UnifiedAPI_BatchOperations : FEntityTestBase
+{
+	virtual bool InstantTest() override
+	{
+		constexpr int32 NumEntities = 10;
+		TArray<FMassEntityHandle> Entities;
+		EntityManager->BatchCreateEntities(IntsArchetype, NumEntities, Entities);
+
+		// Add sparse fragments to all entities using unified API
+		for (const FMassEntityHandle& Entity : Entities)
+		{
+			EntityManager->AddFragmentToEntity(Entity, FTestFragment_SparseInt::StaticStruct());
+			EntityManager->AddTagToEntity(Entity, FTestTag_SparseA::StaticStruct());
+		}
+
+		// Verify all entities have sparse elements
+		for (const FMassEntityHandle& Entity : Entities)
+		{
+			AITEST_TRUE("Each entity should have sparse fragment",
+				EntityManager->DoesEntityHaveElement(Entity, FTestFragment_SparseInt::StaticStruct()));
+			AITEST_TRUE("Each entity should have sparse tag",
+				EntityManager->DoesEntityHaveElement(Entity, FTestTag_SparseA::StaticStruct()));
+		}
+
+		// Remove from half
+		for (int32 i = 0; i < NumEntities / 2; ++i)
+		{
+			EntityManager->RemoveFragmentFromEntity(Entities[i], FTestFragment_SparseInt::StaticStruct());
+			EntityManager->RemoveTagFromEntity(Entities[i], FTestTag_SparseA::StaticStruct());
+		}
+
+		// Verify first half don't have, second half still have
+		for (int32 i = 0; i < NumEntities; ++i)
+		{
+			if (i < NumEntities / 2)
+			{
+				AITEST_FALSE("First half should not have sparse fragment",
+					EntityManager->DoesEntityHaveElement(Entities[i], FTestFragment_SparseInt::StaticStruct()));
+				AITEST_FALSE("First half should not have sparse tag",
+					EntityManager->DoesEntityHaveElement(Entities[i], FTestTag_SparseA::StaticStruct()));
+			}
+			else
+			{
+				AITEST_TRUE("Second half should have sparse fragment",
+					EntityManager->DoesEntityHaveElement(Entities[i], FTestFragment_SparseInt::StaticStruct()));
+				AITEST_TRUE("Second half should have sparse tag",
+					EntityManager->DoesEntityHaveElement(Entities[i], FTestTag_SparseA::StaticStruct()));
+			}
+		}
+
+		return true;
+	}
+};
+IMPLEMENT_AI_INSTANT_TEST(FSparseElements_UnifiedAPI_BatchOperations, "System.Mass.Stress.Sparse.UnifiedAPI.BatchOperations");
+
+
+struct FSparseElements_UnifiedAPI_RemoveNonExistent : FEntityTestBase
+{
+	virtual bool InstantTest() override
+	{
+		const FMassEntityHandle Entity = EntityManager->CreateEntity(IntsArchetype);
+
+		// Try to remove sparse elements that were never added - should not crash
+		EntityManager->RemoveFragmentFromEntity(Entity, FTestFragment_SparseInt::StaticStruct());
+		EntityManager->RemoveTagFromEntity(Entity, FTestTag_SparseA::StaticStruct());
+		EntityManager->RemoveElementFromEntity(Entity, FTestFragment_SparseFloat::StaticStruct());
+
+		// Verify entity is still valid
+		AITEST_TRUE("Entity should still be valid after removing non-existent sparse elements",
+			EntityManager->IsEntityValid(Entity));
+
+		return true;
+	}
+};
+IMPLEMENT_AI_INSTANT_TEST(FSparseElements_UnifiedAPI_RemoveNonExistent, "System.Mass.Stress.Sparse.UnifiedAPI.RemoveNonExistent");
+
+
+struct FSparseElements_UnifiedAPI_AddDuplicate : FEntityTestBase
+{
+	virtual bool InstantTest() override
+	{
+		const FMassEntityHandle Entity = EntityManager->CreateEntity(IntsArchetype);
+
+		// Add sparse fragment
+		EntityManager->AddFragmentToEntity(Entity, FTestFragment_SparseInt::StaticStruct());
+
+		// Set a value
+		FMassEntityView EntityView(*EntityManager, Entity);
+		FTestFragment_SparseInt* SparseFragment = EntityView.GetSparseFragmentDataPtr<FTestFragment_SparseInt>();
+		AITEST_NOT_NULL("Should have sparse fragment", SparseFragment);
+		SparseFragment->Value = 123;
+
+		// Try to add the same sparse fragment again
+		EntityManager->AddFragmentToEntity(Entity, FTestFragment_SparseInt::StaticStruct());
+
+		// Verify value is preserved (not reset)
+		SparseFragment = EntityView.GetSparseFragmentDataPtr<FTestFragment_SparseInt>();
+		AITEST_NOT_NULL("Should still have sparse fragment", SparseFragment);
+		AITEST_EQUAL("Sparse fragment value should be preserved after duplicate add", SparseFragment->Value, 123);
+
+		return true;
+	}
+};
+IMPLEMENT_AI_INSTANT_TEST(FSparseElements_UnifiedAPI_AddDuplicate, "System.Mass.Stress.Sparse.UnifiedAPI.AddDuplicate");
+
+
+struct FSparseElements_UnifiedAPI_EntityView_HasFragment_Sparse : FEntityTestBase
+{
+	virtual bool InstantTest() override
+	{
+		const FMassEntityHandle Entity = EntityManager->CreateEntity(IntsArchetype);
+		FMassEntityView EntityView(*EntityManager, Entity);
+
+		// Verify HasFragment works for sparse fragments
+		AITEST_FALSE("HasFragment should return false for absent sparse fragment",
+			EntityView.HasFragment(FTestFragment_SparseInt::StaticStruct()));
+
+		EntityManager->AddFragmentToEntity(Entity, FTestFragment_SparseInt::StaticStruct());
+
+		AITEST_TRUE("HasFragment should return true for present sparse fragment",
+			EntityView.HasFragment(FTestFragment_SparseInt::StaticStruct()));
+
+		// Template version
+		AITEST_TRUE("Template HasFragment should work for sparse fragment",
+			EntityView.HasFragment<FTestFragment_SparseInt>());
+
+		return true;
+	}
+};
+IMPLEMENT_AI_INSTANT_TEST(FSparseElements_UnifiedAPI_EntityView_HasFragment_Sparse, "System.Mass.Stress.Sparse.UnifiedAPI.EntityView.HasFragment");
+
+
+struct FSparseElements_UnifiedAPI_EntityView_HasTag_Sparse : FEntityTestBase
+{
+	virtual bool InstantTest() override
+	{
+		const FMassEntityHandle Entity = EntityManager->CreateEntity(IntsArchetype);
+		FMassEntityView EntityView(*EntityManager, Entity);
+
+		// Verify HasTag works for sparse tags
+		AITEST_FALSE("HasTag should return false for absent sparse tag",
+			EntityView.HasTag(FTestTag_SparseA::StaticStruct()));
+
+		EntityManager->AddTagToEntity(Entity, FTestTag_SparseA::StaticStruct());
+
+		AITEST_TRUE("HasTag should return true for present sparse tag",
+			EntityView.HasTag(FTestTag_SparseA::StaticStruct()));
+
+		// Template version
+		AITEST_TRUE("Template HasTag should work for sparse tag",
+			EntityView.HasTag<FTestTag_SparseA>());
+
+		return true;
+	}
+};
+IMPLEMENT_AI_INSTANT_TEST(FSparseElements_UnifiedAPI_EntityView_HasTag_Sparse, "System.Mass.Stress.Sparse.UnifiedAPI.EntityView.HasTag");
+
+
+struct FSparseElements_UnifiedAPI_FragmentListOperations : FEntityTestBase
+{
+	virtual bool InstantTest() override
+	{
+		const FMassEntityHandle Entity = EntityManager->CreateEntity(IntsArchetype);
+
+		// Create array of sparse fragment types for RemoveFragmentListFromEntity
+		TArray<const UScriptStruct*> SparseFragmentTypes;
+		SparseFragmentTypes.Add(FTestFragment_SparseInt::StaticStruct());
+		SparseFragmentTypes.Add(FTestFragment_SparseFloat::StaticStruct());
+
+		// Add sparse fragments first
+		for (const UScriptStruct* Type : SparseFragmentTypes)
+		{
+			EntityManager->AddFragmentToEntity(Entity, Type);
+		}
+
+		// Verify all added
+		AITEST_TRUE("Should have sparse int", EntityManager->DoesEntityHaveElement<FTestFragment_SparseInt>(Entity));
+		AITEST_TRUE("Should have sparse float", EntityManager->DoesEntityHaveElement<FTestFragment_SparseFloat>(Entity));
+
+		// Remove using RemoveFragmentListFromEntity
+		EntityManager->RemoveFragmentListFromEntity(Entity, SparseFragmentTypes);
+
+		// Verify all removed
+		AITEST_FALSE("Should not have sparse int after list removal",
+			EntityManager->DoesEntityHaveElement<FTestFragment_SparseInt>(Entity));
+		AITEST_FALSE("Should not have sparse float after list removal",
+			EntityManager->DoesEntityHaveElement<FTestFragment_SparseFloat>(Entity));
+
+		return true;
+	}
+};
+IMPLEMENT_AI_INSTANT_TEST(FSparseElements_UnifiedAPI_FragmentListOperations, "System.Mass.Stress.Sparse.UnifiedAPI.FragmentListOperations");
+
+} // namespace
+
+#undef WITH_ENTITY_MOVING_AFFECTING_COMPOSITION
+
+UE_ENABLE_OPTIMIZATION_SHIP
+
+#undef LOCTEXT_NAMESPACE
