[Mass] fixed a minor type casting issue in one of Mass test files, to make clang compilation happy

#rnx

#Synced-CL 51942366

diff --git a/Engine/Source/Developer/MassEntityTestSuite/Private/StressTests/MassStressSharedFragmentsTest.cpp b/Engine/Source/Developer/MassEntityTestSuite/Private/StressTests/MassStressSharedFragmentsTest.cpp
--- a/Engine/Source/Developer/MassEntityTestSuite/Private/StressTests/MassStressSharedFragmentsTest.cpp
+++ b/Engine/Source/Developer/MassEntityTestSuite/Private/StressTests/MassStressSharedFragmentsTest.cpp
@@ -372,7 +372,7 @@ struct FSetSharedFragmentsData_ChunkBoundary : FEntityTestBase
 		SharedValuesA.Sort();
 
 		const int32 EntitiesPerChunk = FMassArchetypeHelper::ArchetypeDataFromHandleChecked(FloatsArchetype).GetNumEntitiesPerChunk();
-		const int32 NumEntitiesToCreate = static_cast<int32>(EntitiesPerChunk * 1.5f); // Create 1.5 chunks worth
+		const int32 NumEntitiesToCreate = FMath::FloorToInt32(static_cast<float>(EntitiesPerChunk) * 1.5f); // Create 1.5 chunks worth
 
 		TArray<FMassEntityHandle> Entities;
 		EntityManager->BatchCreateEntities(FloatsArchetype, SharedValuesA, NumEntitiesToCreate, Entities);
@@ -765,7 +765,7 @@ struct FSetConstSharedFragmentsData_ChunkBoundary : FEntityTestBase
 		SharedValuesA.Sort();
 
 		const int32 EntitiesPerChunk = FMassArchetypeHelper::ArchetypeDataFromHandleChecked(FloatsArchetype).GetNumEntitiesPerChunk();
-		const int32 NumEntitiesToCreate = static_cast<int32>(EntitiesPerChunk * 1.5f); // Create 1.5 chunks worth
+		const int32 NumEntitiesToCreate = FMath::FloorToInt32(static_cast<float>(EntitiesPerChunk) * 1.5f); // Create 1.5 chunks worth
 
 		TArray<FMassEntityHandle> Entities;
 		EntityManager->BatchCreateEntities(FloatsArchetype, SharedValuesA, NumEntitiesToCreate, Entities);
