[Mass] cooked automation tests compilation fix, put some tests in WITH_MASSENTITY_DEBUG since they were using the debug API

#rnx
#Synced-CL 51942366

diff --git a/Engine/Source/Developer/MassEntityTestSuite/Private/StressTests/MassStressSparseElementsTest.cpp b/Engine/Source/Developer/MassEntityTestSuite/Private/StressTests/MassStressSparseElementsTest.cpp
--- a/Engine/Source/Developer/MassEntityTestSuite/Private/StressTests/MassStressSparseElementsTest.cpp
+++ b/Engine/Source/Developer/MassEntityTestSuite/Private/StressTests/MassStressSparseElementsTest.cpp
@@ -1341,7 +1341,7 @@ namespace UE::Mass::Test::SparseElements
 //-----------------------------------------------------------------------------
 // Chunk and Archetype Bitset Tracking Tests
 //-----------------------------------------------------------------------------
-
+#if WITH_MASSENTITY_DEBUG
 struct FSparseElements_ChunkBitset_AddSingle : FEntityTestBase
 {
 	virtual bool InstantTest() override
@@ -1538,6 +1538,7 @@ struct FSparseElements_ChunkBitset_MultipleSparseTypes : FEntityTestBase
 	}
 };
 IMPLEMENT_AI_INSTANT_TEST(FSparseElements_ChunkBitset_MultipleSparseTypes, "System.Mass.Stress.Sparse.ChunkBitset.MultipleSparseTypes");
+#endif // WITH_MASSENTITY_DEBUG
 
 struct FSparseElements_ChunkBitset_CrossChunk : FEntityTestBase
 {
@@ -1655,6 +1656,7 @@ struct FSparseElements_ArchetypeBitset_MultipleChunks : FEntityTestBase
 };
 IMPLEMENT_AI_INSTANT_TEST(FSparseElements_ArchetypeBitset_MultipleChunks, "System.Mass.Stress.Sparse.ArchetypeBitset.MultipleChunks");
 
+#if WITH_MASSENTITY_DEBUG
 struct FSparseElements_EntityRemoval_ChunkBitsetUpdate : FEntityTestBase
 {
 	virtual bool InstantTest() override
@@ -1739,6 +1741,7 @@ struct FSparseElements_BatchEntityDestruction_ChunkBitsetUpdate : FEntityTestBas
 	}
 };
 IMPLEMENT_AI_INSTANT_TEST(FSparseElements_BatchEntityDestruction_ChunkBitsetUpdate, "System.Mass.Stress.Sparse.EntityRemoval.BatchDestruction");
+#endif // WITH_MASSENTITY_DEBUG
 
 struct FSparseElements_ArchetypeChange_BitsetTracking : FEntityTestBase
 {
@@ -1873,6 +1876,7 @@ struct FSparseElements_ComplexScenario_BitsetAccuracy : FEntityTestBase
 };
 IMPLEMENT_AI_INSTANT_TEST(FSparseElements_ComplexScenario_BitsetAccuracy, "System.Mass.Stress.Sparse.Complex.BitsetAccuracy");
 
+#if WITH_MASSENTITY_DEBUG
 struct FSparseElements_ReAdd_ChunkBitsetConsistency : FEntityTestBase
 {
 	virtual bool InstantTest() override
@@ -1916,7 +1920,7 @@ struct FSparseElements_ReAdd_ChunkBitsetConsistency : FEntityTestBase
 	}
 };
 IMPLEMENT_AI_INSTANT_TEST(FSparseElements_ReAdd_ChunkBitsetConsistency, "System.Mass.Stress.Sparse.ReAdd.ChunkBitsetConsistency");
-
+#endif // WITH_MASSENTITY_DEBUG
 
 //----------------------------------------------------------------------//
 // Unified API Tests - verifying Fragment/Tag API works seamlessly with Sparse Elements
