[Mass] Add MassEntity API test coverage for untested APIs across execution context, query chunk filtering, command buffer, processor activation, entity view, composition descriptors, shared fragment values, entity manager, observer batch notification, archetype internals, processor config, relations, and serialization

The tests exposed a misleading comment for shared fragments - that has been fixed as well.

#rnx
#Synced-CL 51998236

diff --git a/Engine/Source/Developer/MassEntityTestSuite/Private/Coverage/MassArchetypeTest.cpp b/Engine/Source/Developer/MassEntityTestSuite/Private/Coverage/MassArchetypeTest.cpp
new file mode 100644
--- /dev/null
+++ b/Engine/Source/Developer/MassEntityTestSuite/Private/Coverage/MassArchetypeTest.cpp
@@ -0,0 +1,100 @@
+// Copyright Epic Games, Inc. All Rights Reserved.
+
+#include "AITestsCommon.h"
+#include "MassEntityManager.h"
+#include "MassEntityTestTypes.h"
+#include "MassArchetypeTypes.h"
+#include "MassEntityTypes.h"
+#include "MassExecutionContext.h"
+#include "MassEntityQuery.h"
+
+#define LOCTEXT_NAMESPACE "MassTest"
+
+UE_DISABLE_OPTIMIZATION_SHIP
+
+namespace UE::Mass::Test::Coverage
+{
+
+struct FArchetype_VersionedHandleInvalidation : FEntityTestBase
+{
+	virtual bool InstantTest() override
+	{
+		CA_ASSUME(EntityManager);
+
+		TArray<FMassEntityHandle> Entities;
+		EntityManager->BatchCreateEntities(IntsArchetype, 10, Entities);
+
+		// Get a versioned handle
+		FMassArchetypeVersionedHandle VersionedHandle(IntsArchetype);
+
+		AITEST_TRUE("Versioned handle is valid initially", VersionedHandle.IsValid());
+
+		// Mutate entities — destroying should change entity order version
+		EntityManager->DestroyEntity(Entities[0]);
+		EntityManager->DestroyEntity(Entities[5]);
+
+		// The versioned handle should now be out of date (entity order changed)
+		AITEST_FALSE("Versioned handle is NOT up to date after entity destruction", VersionedHandle.IsUpToDate());
+
+		return true;
+	}
+};
+IMPLEMENT_AI_INSTANT_TEST(FArchetype_VersionedHandleInvalidation, "System.Mass.Coverage.Archetype.VersionedHandleInvalidation");
+
+struct FArchetype_SharedFragmentPerChunkStorage : FEntityTestBase
+{
+	virtual bool InstantTest() override
+	{
+		CA_ASSUME(EntityManager);
+
+		// Create archetype with fragment + const shared fragment type
+		FMassArchetypeCompositionDescriptor Composition;
+		Composition.Add<FTestFragment_Int>();
+		Composition.Add<FTestConstSharedFragment_Int>();
+		const FMassArchetypeHandle Archetype = EntityManager->CreateArchetype(Composition);
+
+		// Create entities with DIFFERENT shared values — they should be in the SAME archetype
+		// (shared fragment TYPES define archetype composition, VALUES are per-chunk)
+		FMassArchetypeSharedFragmentValues SharedValuesA;
+		SharedValuesA.Add(FConstSharedStruct::Make<FTestConstSharedFragment_Int>(42));
+		FMassArchetypeSharedFragmentValues SharedValuesB;
+		SharedValuesB.Add(FConstSharedStruct::Make<FTestConstSharedFragment_Int>(99));
+
+		TArray<FMassEntityHandle> EntitiesA;
+		EntityManager->BatchCreateEntities(Archetype, SharedValuesA, 5, EntitiesA);
+		TArray<FMassEntityHandle> EntitiesB;
+		EntityManager->BatchCreateEntities(Archetype, SharedValuesB, 5, EntitiesB);
+
+		// Different shared values → same archetype (archetype is defined by types, not values)
+		const FMassArchetypeHandle ArchA = EntityManager->GetArchetypeForEntity(EntitiesA[0]);
+		const FMassArchetypeHandle ArchB = EntityManager->GetArchetypeForEntity(EntitiesB[0]);
+		AITEST_EQUAL("Different shared values produce SAME archetype", ArchA, ArchB);
+
+		// But queries observe distinct shared values per chunk
+		FMassEntityQuery Query(EntityManager);
+		Query.AddRequirement<FTestFragment_Int>(EMassFragmentAccess::ReadOnly);
+		Query.AddConstSharedRequirement<FTestConstSharedFragment_Int>();
+
+		TSet<int32> ObservedSharedValues;
+		FMassExecutionContext ExecutionContext(*EntityManager);
+		Query.ForEachEntityChunk(ExecutionContext,
+			[&ObservedSharedValues](FMassExecutionContext& Context)
+		{
+			const FTestConstSharedFragment_Int& SharedFrag = Context.GetConstSharedFragment<FTestConstSharedFragment_Int>();
+			ObservedSharedValues.Add(SharedFrag.Value);
+		});
+
+		AITEST_EQUAL("Query observed 2 distinct shared values across chunks", ObservedSharedValues.Num(), 2);
+		AITEST_TRUE("Observed value 42", ObservedSharedValues.Contains(42));
+		AITEST_TRUE("Observed value 99", ObservedSharedValues.Contains(99));
+
+		return true;
+	}
+};
+IMPLEMENT_AI_INSTANT_TEST(FArchetype_SharedFragmentPerChunkStorage, "System.Mass.Coverage.Archetype.SharedFragmentPerChunkStorage");
+
+} // UE::Mass::Test::Coverage
+
+UE_ENABLE_OPTIMIZATION_SHIP
+
+#undef LOCTEXT_NAMESPACE
diff --git a/Engine/Source/Developer/MassEntityTestSuite/Private/Coverage/MassCommandBufferTest.cpp b/Engine/Source/Developer/MassEntityTestSuite/Private/Coverage/MassCommandBufferTest.cpp
new file mode 100644
--- /dev/null
+++ b/Engine/Source/Developer/MassEntityTestSuite/Private/Coverage/MassCommandBufferTest.cpp
@@ -0,0 +1,110 @@
+// Copyright Epic Games, Inc. All Rights Reserved.
+
+#include "AITestsCommon.h"
+#include "MassEntityManager.h"
+#include "MassEntityTestTypes.h"
+#include "MassCommandBuffer.h"
+#include "MassEntityView.h"
+
+#define LOCTEXT_NAMESPACE "MassTest"
+
+UE_DISABLE_OPTIMIZATION_SHIP
+
+namespace UE::Mass::Test::Coverage
+{
+
+struct FCommands_CancelCommands : FEntityTestBase
+{
+	virtual bool InstantTest() override
+	{
+		CA_ASSUME(EntityManager);
+
+		const FMassEntityHandle Entity = EntityManager->CreateEntity(FloatsArchetype);
+
+		// Queue commands
+		EntityManager->Defer().AddFragment<FTestFragment_Int>(Entity);
+		EntityManager->Defer().AddTag<FTestTag_A>(Entity);
+
+		// Cancel all pending commands
+		EntityManager->Defer().CancelCommands();
+
+		// Flush — should be a no-op
+		EntityManager->FlushCommands();
+
+		// Entity should still be in FloatsArchetype with no Int fragment and no tag
+		AITEST_EQUAL("Entity still in original archetype", EntityManager->GetArchetypeForEntity(Entity), FloatsArchetype);
+
+		FMassEntityView View(*EntityManager, Entity);
+		AITEST_NULL("FTestFragment_Int should NOT be present", View.GetFragmentDataPtr<FTestFragment_Int>());
+		AITEST_FALSE("FTestTag_A should NOT be present", View.HasTag<FTestTag_A>());
+
+		return true;
+	}
+};
+IMPLEMENT_AI_INSTANT_TEST(FCommands_CancelCommands, "System.Mass.Coverage.CommandBuffer.CancelCommands");
+
+struct FCommands_MoveAppend : FEntityTestBase
+{
+	virtual bool InstantTest() override
+	{
+		CA_ASSUME(EntityManager);
+
+		const FMassEntityHandle Entity = EntityManager->CreateEntity(EmptyArchetype);
+
+		// Use the entity manager's deferred command buffer to build up commands,
+		// then test MoveAppend by creating a separate buffer and merging it in.
+		// Note: FMassCommandBuffer() has a default constructor only, and Flush() is private
+		// (only callable by FMassEntityManager). So we work through EntityManager->Defer().
+		FMassCommandBuffer SeparateBuffer;
+		SeparateBuffer.AddFragment<FTestFragment_Int>(Entity);
+
+		AITEST_TRUE("Separate buffer has pending commands", SeparateBuffer.HasPendingCommands());
+
+		// Add Float via the manager's deferred buffer
+		EntityManager->Defer().AddFragment<FTestFragment_Float>(Entity);
+
+		// MoveAppend the separate buffer into the manager's deferred buffer
+		EntityManager->Defer().MoveAppend(SeparateBuffer);
+
+		AITEST_FALSE("Separate buffer is empty after MoveAppend", SeparateBuffer.HasPendingCommands());
+		AITEST_TRUE("Deferred buffer has pending commands after MoveAppend", EntityManager->Defer().HasPendingCommands());
+
+		// Flush via the entity manager — both commands should execute
+		EntityManager->FlushCommands();
+
+		FMassEntityView View(*EntityManager, Entity);
+		AITEST_NOT_NULL("Entity has FTestFragment_Float after flush", View.GetFragmentDataPtr<FTestFragment_Float>());
+		AITEST_NOT_NULL("Entity has FTestFragment_Int after flush", View.GetFragmentDataPtr<FTestFragment_Int>());
+
+		return true;
+	}
+};
+IMPLEMENT_AI_INSTANT_TEST(FCommands_MoveAppend, "System.Mass.Coverage.CommandBuffer.MoveAppend");
+
+struct FCommands_HasPendingCommands : FEntityTestBase
+{
+	virtual bool InstantTest() override
+	{
+		CA_ASSUME(EntityManager);
+
+		AITEST_FALSE("No pending commands initially", EntityManager->Defer().HasPendingCommands());
+
+		const FMassEntityHandle Entity = EntityManager->CreateEntity(FloatsArchetype);
+		EntityManager->Defer().AddTag<FTestTag_A>(Entity);
+
+		AITEST_TRUE("Has pending commands after AddTag", EntityManager->Defer().HasPendingCommands());
+
+		EntityManager->FlushCommands();
+
+		AITEST_FALSE("No pending commands after flush", EntityManager->Defer().HasPendingCommands());
+
+		return true;
+	}
+};
+IMPLEMENT_AI_INSTANT_TEST(FCommands_HasPendingCommands, "System.Mass.Coverage.CommandBuffer.HasPendingCommands");
+
+} // UE::Mass::Test::Coverage
+
+UE_ENABLE_OPTIMIZATION_SHIP
+
+#undef LOCTEXT_NAMESPACE
diff --git a/Engine/Source/Developer/MassEntityTestSuite/Private/Coverage/MassCompositionDescriptorTest.cpp b/Engine/Source/Developer/MassEntityTestSuite/Private/Coverage/MassCompositionDescriptorTest.cpp
new file mode 100644
--- /dev/null
+++ b/Engine/Source/Developer/MassEntityTestSuite/Private/Coverage/MassCompositionDescriptorTest.cpp
@@ -0,0 +1,106 @@
+// Copyright Epic Games, Inc. All Rights Reserved.
+
+#include "AITestsCommon.h"
+#include "MassEntityManager.h"
+#include "MassEntityTestTypes.h"
+#include "MassEntityTypes.h"
+
+#define LOCTEXT_NAMESPACE "MassTest"
+
+UE_DISABLE_OPTIMIZATION_SHIP
+
+namespace UE::Mass::Test::Coverage
+{
+
+struct FComposition_Difference : FEntityTestBase
+{
+	virtual bool InstantTest() override
+	{
+		// A = {FragFloat, FragInt, TagA}
+		FMassArchetypeCompositionDescriptor A;
+		A.Add<FTestFragment_Float>();
+		A.Add<FTestFragment_Int>();
+		A.Add<FTestTag_A>();
+
+		// B = {FragInt, FragBool, TagA, TagB}
+		FMassArchetypeCompositionDescriptor B;
+		B.Add<FTestFragment_Int>();
+		B.Add<FTestFragment_Bool>();
+		B.Add<FTestTag_A>();
+		B.Add<FTestTag_B>();
+
+		// A.CalculateDifference(B) = elements in A not in B = {FragFloat}
+		const FMassArchetypeCompositionDescriptor AMinusB = A.CalculateDifference(B);
+		AITEST_TRUE("A-B contains FragFloat", AMinusB.Contains<FTestFragment_Float>());
+		AITEST_FALSE("A-B does not contain FragInt", AMinusB.Contains<FTestFragment_Int>());
+		AITEST_FALSE("A-B does not contain TagA", AMinusB.Contains<FTestTag_A>());
+		AITEST_FALSE("A-B does not contain FragBool", AMinusB.Contains<FTestFragment_Bool>());
+		AITEST_EQUAL("A-B has 1 stored type", AMinusB.CountStoredTypes(), 1);
+
+		// B.CalculateDifference(A) = elements in B not in A = {FragBool, TagB}
+		const FMassArchetypeCompositionDescriptor BMinusA = B.CalculateDifference(A);
+		AITEST_TRUE("B-A contains FragBool", BMinusA.Contains<FTestFragment_Bool>());
+		AITEST_TRUE("B-A contains TagB", BMinusA.Contains<FTestTag_B>());
+		AITEST_FALSE("B-A does not contain FragInt", BMinusA.Contains<FTestFragment_Int>());
+		AITEST_EQUAL("B-A has 2 stored types", BMinusA.CountStoredTypes(), 2);
+
+		return true;
+	}
+};
+IMPLEMENT_AI_INSTANT_TEST(FComposition_Difference, "System.Mass.Coverage.Composition.Difference");
+
+struct FComposition_HasAll : FEntityTestBase
+{
+	virtual bool InstantTest() override
+	{
+		FMassArchetypeCompositionDescriptor Full;
+		Full.Add<FTestFragment_Float>();
+		Full.Add<FTestFragment_Int>();
+		Full.Add<FTestTag_A>();
+
+		FMassArchetypeCompositionDescriptor Subset;
+		Subset.Add<FTestFragment_Float>();
+		Subset.Add<FTestTag_A>();
+
+		FMassArchetypeCompositionDescriptor NonSubset;
+		NonSubset.Add<FTestFragment_Float>();
+		NonSubset.Add<FTestFragment_Bool>(); // Not in Full
+
+		AITEST_TRUE("Full HasAll of valid subset", Full.HasAll(Subset));
+		AITEST_FALSE("Full does NOT HasAll of non-subset", Full.HasAll(NonSubset));
+
+		// Empty descriptor is a subset of everything
+		FMassArchetypeCompositionDescriptor Empty;
+		AITEST_TRUE("Full HasAll of empty", Full.HasAll(Empty));
+
+		return true;
+	}
+};
+IMPLEMENT_AI_INSTANT_TEST(FComposition_HasAll, "System.Mass.Coverage.Composition.HasAll");
+
+struct FComposition_CountStoredTypes : FEntityTestBase
+{
+	virtual bool InstantTest() override
+	{
+		FMassArchetypeCompositionDescriptor Desc;
+		AITEST_EQUAL("Empty composition has 0 types", Desc.CountStoredTypes(), 0);
+
+		Desc.Add<FTestFragment_Float>();
+		Desc.Add<FTestFragment_Int>();
+		Desc.Add<FTestFragment_Bool>();
+		Desc.Add<FTestTag_A>();
+		Desc.Add<FTestTag_B>();
+		Desc.Add<FTestSharedFragment_Int>();
+
+		AITEST_EQUAL("Composition with 3 fragments + 2 tags + 1 shared = 6 types", Desc.CountStoredTypes(), 6);
+
+		return true;
+	}
+};
+IMPLEMENT_AI_INSTANT_TEST(FComposition_CountStoredTypes, "System.Mass.Coverage.Composition.CountStoredTypes");
+
+} // UE::Mass::Test::Coverage
+
+UE_ENABLE_OPTIMIZATION_SHIP
+
+#undef LOCTEXT_NAMESPACE
diff --git a/Engine/Source/Developer/MassEntityTestSuite/Private/Coverage/MassEntityManagerTest.cpp b/Engine/Source/Developer/MassEntityTestSuite/Private/Coverage/MassEntityManagerTest.cpp
new file mode 100644
--- /dev/null
+++ b/Engine/Source/Developer/MassEntityTestSuite/Private/Coverage/MassEntityManagerTest.cpp
@@ -0,0 +1,186 @@
+// Copyright Epic Games, Inc. All Rights Reserved.
+
+#include "AITestsCommon.h"
+#include "MassEntityManager.h"
+#include "MassEntityTestTypes.h"
+#include "MassEntityView.h"
+#include "MassEntityTypes.h"
+#include "MassExecutionContext.h"
+
+#define LOCTEXT_NAMESPACE "MassTest"
+
+UE_DISABLE_OPTIMIZATION_SHIP
+
+namespace UE::Mass::Test::Coverage
+{
+
+struct FManager_EntityCompaction : FEntityTestBase
+{
+	virtual bool InstantTest() override
+	{
+		CA_ASSUME(EntityManager);
+
+		const int32 EntitiesPerChunk = EntityManager->DebugGetArchetypeEntitiesCountPerChunk(IntsArchetype);
+		// Create enough entities to span at least 3 chunks
+		const int32 TotalEntities = EntitiesPerChunk * 3;
+		TArray<FMassEntityHandle> Entities;
+		EntityManager->BatchCreateEntities(IntsArchetype, TotalEntities, Entities);
+
+		// Set values so we can verify data integrity after compaction
+		for (int32 i = 0; i < TotalEntities; ++i)
+		{
+			EntityManager->GetFragmentDataChecked<FTestFragment_Int>(Entities[i]).Value = i;
+		}
+
+		// Destroy every other entity to create fragmentation
+		TArray<FMassEntityHandle> SurvivorEntities;
+		for (int32 i = 0; i < TotalEntities; ++i)
+		{
+			if (i % 2 == 0)
+			{
+				EntityManager->DestroyEntity(Entities[i]);
+			}
+			else
+			{
+				SurvivorEntities.Add(Entities[i]);
+			}
+		}
+
+		const int32 SurvivorCount = SurvivorEntities.Num();
+
+		// Compact with unlimited time budget
+		EntityManager->DoEntityCompaction(TNumericLimits<double>::Max());
+
+		// Verify all surviving entities have intact data
+		for (int32 i = 0; i < SurvivorCount; ++i)
+		{
+			AITEST_TRUE("Survivor entity is still valid", EntityManager->IsEntityValid(SurvivorEntities[i]));
+			const int32 OriginalIndex = i * 2 + 1; // Odd indices survived
+			const int32 ActualValue = EntityManager->GetFragmentDataChecked<FTestFragment_Int>(SurvivorEntities[i]).Value;
+			AITEST_EQUAL("Data integrity after compaction", ActualValue, OriginalIndex);
+		}
+
+		return true;
+	}
+};
+IMPLEMENT_AI_INSTANT_TEST(FManager_EntityCompaction, "System.Mass.Coverage.Manager.EntityCompaction");
+
+struct FManager_SwapTagsForEntity : FEntityTestBase
+{
+	virtual bool InstantTest() override
+	{
+		CA_ASSUME(EntityManager);
+
+		const FMassEntityHandle Entity = EntityManager->CreateEntity(IntsArchetype);
+
+		// Add TagA
+		EntityManager->AddTagToEntity(Entity, FTestTag_A::StaticStruct());
+		{
+			FMassEntityView View(*EntityManager, Entity);
+			AITEST_TRUE("Entity has TagA before swap", View.HasTag<FTestTag_A>());
+			AITEST_FALSE("Entity does not have TagB before swap", View.HasTag<FTestTag_B>());
+		}
+
+		// Swap TagA -> TagB
+		EntityManager->SwapTagsForEntity(Entity, FTestTag_A::StaticStruct(), FTestTag_B::StaticStruct());
+
+		{
+			FMassEntityView View(*EntityManager, Entity);
+			AITEST_FALSE("Entity does not have TagA after swap", View.HasTag<FTestTag_A>());
+			AITEST_TRUE("Entity has TagB after swap", View.HasTag<FTestTag_B>());
+		}
+
+		return true;
+	}
+};
+IMPLEMENT_AI_INSTANT_TEST(FManager_SwapTagsForEntity, "System.Mass.Coverage.Manager.SwapTagsForEntity");
+
+struct FManager_ArchetypeTraversal : FEntityTestBase
+{
+	virtual bool InstantTest() override
+	{
+		CA_ASSUME(EntityManager);
+
+		// Create archetypes
+		EntityManager->CreateEntity(FloatsArchetype);
+		EntityManager->CreateEntity(IntsArchetype);
+		EntityManager->CreateEntity(FloatsIntsArchetype);
+
+		// Traverse FloatsIntsArchetype and collect fragment types
+		TArray<const UScriptStruct*> CollectedTypes;
+		FMassEntityManager::ForEachArchetypeFragmentType(FloatsIntsArchetype, [&CollectedTypes](const UScriptStruct* FragmentType)
+		{
+			CollectedTypes.Add(FragmentType);
+		});
+
+		AITEST_EQUAL("FloatsInts archetype has 2 fragment types", CollectedTypes.Num(), 2);
+
+		const bool bHasFloat = CollectedTypes.Contains(FTestFragment_Float::StaticStruct());
+		const bool bHasInt = CollectedTypes.Contains(FTestFragment_Int::StaticStruct());
+		AITEST_TRUE("Contains FTestFragment_Float", bHasFloat);
+		AITEST_TRUE("Contains FTestFragment_Int", bHasInt);
+
+		return true;
+	}
+};
+IMPLEMENT_AI_INSTANT_TEST(FManager_ArchetypeTraversal, "System.Mass.Coverage.Manager.ArchetypeTraversal");
+
+struct FManager_BatchAddElements : FEntityTestBase
+{
+	virtual bool InstantTest() override
+	{
+		CA_ASSUME(EntityManager);
+
+		const int32 NumEntities = 100;
+		TArray<FMassEntityHandle> Entities;
+		EntityManager->BatchCreateEntities(FloatsArchetype, NumEntities, Entities);
+
+		// Verify all start without Int fragment
+		for (const FMassEntityHandle& Entity : Entities)
+		{
+			AITEST_EQUAL("Initial archetype is Floats", EntityManager->GetArchetypeForEntity(Entity), FloatsArchetype);
+		}
+
+		// Batch add FTestFragment_Int to all entities
+		EntityManager->BatchAddElementToEntities(Entities, FTestFragment_Int::StaticStruct());
+
+		// All entities should now be in FloatsInts archetype
+		for (const FMassEntityHandle& Entity : Entities)
+		{
+			AITEST_EQUAL("Migrated to FloatsInts archetype", EntityManager->GetArchetypeForEntity(Entity), FloatsIntsArchetype);
+		}
+
+		return true;
+	}
+};
+IMPLEMENT_AI_INSTANT_TEST(FManager_BatchAddElements, "System.Mass.Coverage.Manager.BatchAddElements");
+
+struct FManager_BatchRemoveElements : FEntityTestBase
+{
+	virtual bool InstantTest() override
+	{
+		CA_ASSUME(EntityManager);
+
+		const int32 NumEntities = 100;
+		TArray<FMassEntityHandle> Entities;
+		EntityManager->BatchCreateEntities(FloatsIntsArchetype, NumEntities, Entities);
+
+		// Batch remove FTestFragment_Int from all entities
+		EntityManager->BatchRemoveElementFromEntities(Entities, FTestFragment_Int::StaticStruct());
+
+		// All entities should now be in Floats-only archetype
+		for (const FMassEntityHandle& Entity : Entities)
+		{
+			AITEST_EQUAL("Migrated to Floats archetype", EntityManager->GetArchetypeForEntity(Entity), FloatsArchetype);
+		}
+
+		return true;
+	}
+};
+IMPLEMENT_AI_INSTANT_TEST(FManager_BatchRemoveElements, "System.Mass.Coverage.Manager.BatchRemoveElements");
+
+} // UE::Mass::Test::Coverage
+
+UE_ENABLE_OPTIMIZATION_SHIP
+
+#undef LOCTEXT_NAMESPACE
diff --git a/Engine/Source/Developer/MassEntityTestSuite/Private/Coverage/MassEntityQueryTest.cpp b/Engine/Source/Developer/MassEntityTestSuite/Private/Coverage/MassEntityQueryTest.cpp
new file mode 100644
--- /dev/null
+++ b/Engine/Source/Developer/MassEntityTestSuite/Private/Coverage/MassEntityQueryTest.cpp
@@ -0,0 +1,210 @@
+// Copyright Epic Games, Inc. All Rights Reserved.
+
+#include "AITestsCommon.h"
+#include "MassEntityManager.h"
+#include "MassEntityTestTypes.h"
+#include "MassExecutionContext.h"
+#include "MassEntityQuery.h"
+#include "MassArchetypeTypes.h"
+
+#define LOCTEXT_NAMESPACE "MassTest"
+
+UE_DISABLE_OPTIMIZATION_SHIP
+
+namespace UE::Mass::Test::Coverage
+{
+
+struct FQuery_ChunkFilter_Basic : FEntityTestBase
+{
+	virtual bool InstantTest() override
+	{
+		CA_ASSUME(EntityManager);
+
+		// Create archetype with chunk fragment
+		const UScriptStruct* Types[] = { FTestFragment_Int::StaticStruct(), FTestChunkFragment_Int::StaticStruct() };
+		const FMassArchetypeHandle Archetype = EntityManager->CreateArchetype(MakeArrayView(Types));
+
+		// Find how many entities fit per chunk, then create enough to span multiple chunks
+		const int32 EntitiesPerChunk = EntityManager->DebugGetArchetypeEntitiesCountPerChunk(Archetype);
+		const int32 NumChunks = 3;
+		const int32 TotalEntities = EntitiesPerChunk * NumChunks;
+		TArray<FMassEntityHandle> Entities;
+		EntityManager->BatchCreateEntities(Archetype, TotalEntities, Entities);
+
+		// Set up query with chunk filter
+		FMassEntityQuery Query(EntityManager);
+		Query.AddRequirement<FTestFragment_Int>(EMassFragmentAccess::ReadOnly);
+		Query.AddChunkRequirement<FTestChunkFragment_Int>(EMassFragmentAccess::ReadWrite);
+
+		// Set chunk fragment values to distinguish chunks — we'll write a value per chunk
+		// First iteration: set chunk fragment values
+		{
+			int32 ChunkIndex = 0;
+			FMassExecutionContext SetupContext(*EntityManager);
+			Query.ForEachEntityChunk(SetupContext, [&ChunkIndex](FMassExecutionContext& Context)
+			{
+				FTestChunkFragment_Int& ChunkFrag = Context.GetMutableChunkFragment<FTestChunkFragment_Int>();
+				ChunkFrag.Value = ChunkIndex;
+				ChunkIndex++;
+			});
+		}
+
+		// Now set chunk filter: only accept chunks where ChunkFragment value > 0
+		const int32 Threshold = 0;
+		Query.SetChunkFilter([Threshold](const FMassExecutionContext& Context) -> bool
+		{
+			const FTestChunkFragment_Int& ChunkFrag = Context.GetChunkFragment<FTestChunkFragment_Int>();
+			return ChunkFrag.Value > Threshold;
+		});
+
+		AITEST_TRUE("HasChunkFilter returns true after SetChunkFilter", Query.HasChunkFilter());
+
+		// Count entities visited with filter
+		int32 FilteredEntityCount = 0;
+		{
+			FMassExecutionContext FilterContext(*EntityManager);
+			Query.ForEachEntityChunk(FilterContext, [&FilteredEntityCount](FMassExecutionContext& Context)
+			{
+				FilteredEntityCount += Context.GetNumEntities();
+			});
+		}
+
+		// Chunk 0 has Value=0 (filtered out), chunks 1 and 2 have Value>0 (kept)
+		const int32 ExpectedEntities = EntitiesPerChunk * (NumChunks - 1);
+		AITEST_EQUAL("Only entities in chunks passing filter were visited", FilteredEntityCount, ExpectedEntities);
+
+		return true;
+	}
+};
+IMPLEMENT_AI_INSTANT_TEST(FQuery_ChunkFilter_Basic, "System.Mass.Coverage.Query.ChunkFilter.Basic");
+
+struct FQuery_ChunkFilter_Clear : FEntityTestBase
+{
+	virtual bool InstantTest() override
+	{
+		CA_ASSUME(EntityManager);
+
+		const UScriptStruct* Types[] = { FTestFragment_Int::StaticStruct(), FTestChunkFragment_Int::StaticStruct() };
+		const FMassArchetypeHandle Archetype = EntityManager->CreateArchetype(MakeArrayView(Types));
+
+		const int32 EntitiesPerChunk = EntityManager->DebugGetArchetypeEntitiesCountPerChunk(Archetype);
+		const int32 TotalEntities = EntitiesPerChunk * 2;
+		TArray<FMassEntityHandle> Entities;
+		EntityManager->BatchCreateEntities(Archetype, TotalEntities, Entities);
+
+		FMassEntityQuery Query(EntityManager);
+		Query.AddRequirement<FTestFragment_Int>(EMassFragmentAccess::ReadOnly);
+		Query.AddChunkRequirement<FTestChunkFragment_Int>(EMassFragmentAccess::ReadWrite);
+
+		// Set chunk 0 value to 0, chunk 1 value to 1
+		{
+			int32 ChunkIndex = 0;
+			FMassExecutionContext SetupContext(*EntityManager);
+			Query.ForEachEntityChunk(SetupContext, [&ChunkIndex](FMassExecutionContext& Context)
+			{
+				Context.GetMutableChunkFragment<FTestChunkFragment_Int>().Value = ChunkIndex++;
+			});
+		}
+
+		AITEST_FALSE("No chunk filter initially", Query.HasChunkFilter());
+
+		// Set filter to skip chunk 0
+		Query.SetChunkFilter([](const FMassExecutionContext& Context) -> bool
+		{
+			return Context.GetChunkFragment<FTestChunkFragment_Int>().Value > 0;
+		});
+		AITEST_TRUE("Has chunk filter after set", Query.HasChunkFilter());
+
+		int32 FilteredCount = 0;
+		{
+			FMassExecutionContext ExecContext(*EntityManager);
+			Query.ForEachEntityChunk(ExecContext, [&FilteredCount](FMassExecutionContext& Context)
+			{
+				FilteredCount += Context.GetNumEntities();
+			});
+		}
+		AITEST_EQUAL("Filtered count is one chunk", FilteredCount, EntitiesPerChunk);
+
+		// Clear filter
+		Query.ClearChunkFilter();
+		AITEST_FALSE("No chunk filter after clear", Query.HasChunkFilter());
+
+		int32 UnfilteredCount = 0;
+		{
+			FMassExecutionContext ExecContext(*EntityManager);
+			Query.ForEachEntityChunk(ExecContext, [&UnfilteredCount](FMassExecutionContext& Context)
+			{
+				UnfilteredCount += Context.GetNumEntities();
+			});
+		}
+		AITEST_EQUAL("Unfiltered count is all entities", UnfilteredCount, TotalEntities);
+
+		return true;
+	}
+};
+IMPLEMENT_AI_INSTANT_TEST(FQuery_ChunkFilter_Clear, "System.Mass.Coverage.Query.ChunkFilter.Clear");
+
+struct FQuery_ChunkFilter_WithCollections : FEntityTestBase
+{
+	virtual bool InstantTest() override
+	{
+		CA_ASSUME(EntityManager);
+
+		const UScriptStruct* Types[] = { FTestFragment_Int::StaticStruct(), FTestChunkFragment_Int::StaticStruct() };
+		const FMassArchetypeHandle Archetype = EntityManager->CreateArchetype(MakeArrayView(Types));
+
+		const int32 EntitiesPerChunk = EntityManager->DebugGetArchetypeEntitiesCountPerChunk(Archetype);
+		const int32 TotalEntities = EntitiesPerChunk * 3;
+		TArray<FMassEntityHandle> Entities;
+		EntityManager->BatchCreateEntities(Archetype, TotalEntities, Entities);
+
+		FMassEntityQuery Query(EntityManager);
+		Query.AddRequirement<FTestFragment_Int>(EMassFragmentAccess::ReadOnly);
+		Query.AddChunkRequirement<FTestChunkFragment_Int>(EMassFragmentAccess::ReadWrite);
+
+		// Tag chunks with indices
+		{
+			int32 ChunkIndex = 0;
+			FMassExecutionContext SetupContext(*EntityManager);
+			Query.ForEachEntityChunk(SetupContext, [&ChunkIndex](FMassExecutionContext& Context)
+			{
+				Context.GetMutableChunkFragment<FTestChunkFragment_Int>().Value = ChunkIndex++;
+			});
+		}
+
+		// Create a collection from a subset of entities (first 2 chunks)
+		TArray<FMassEntityHandle> SubsetEntities(Entities.GetData(), EntitiesPerChunk * 2);
+		FMassArchetypeEntityCollection Collection(Archetype, SubsetEntities, FMassArchetypeEntityCollection::NoDuplicates);
+
+		// Set chunk filter to only accept chunk index > 0
+		Query.SetChunkFilter([](const FMassExecutionContext& Context) -> bool
+		{
+			return Context.GetChunkFragment<FTestChunkFragment_Int>().Value > 0;
+		});
+
+		// Execute on collection with chunk filter — should only visit chunk 1 entities
+		int32 VisitedCount = 0;
+		{
+			FMassExecutionContext ExecContext(*EntityManager);
+			TArray<FMassArchetypeEntityCollection> Collections;
+			Collections.Add(MoveTemp(Collection));
+			Query.ForEachEntityChunkInCollections(MakeArrayView(Collections), ExecContext, [&VisitedCount](FMassExecutionContext& Context)
+			{
+				VisitedCount += Context.GetNumEntities();
+			});
+		}
+
+		// Collection has chunks 0 and 1. Filter removes chunk 0. Only chunk 1 remains.
+		AITEST_EQUAL("Only entities in collection AND passing chunk filter were visited", VisitedCount, EntitiesPerChunk);
+
+		Query.ClearChunkFilter();
+		return true;
+	}
+};
+IMPLEMENT_AI_INSTANT_TEST(FQuery_ChunkFilter_WithCollections, "System.Mass.Coverage.Query.ChunkFilter.WithCollections");
+
+} // UE::Mass::Test::Coverage
+
+UE_ENABLE_OPTIMIZATION_SHIP
+
+#undef LOCTEXT_NAMESPACE
diff --git a/Engine/Source/Developer/MassEntityTestSuite/Private/Coverage/MassEntityViewCoverageTest.cpp b/Engine/Source/Developer/MassEntityTestSuite/Private/Coverage/MassEntityViewCoverageTest.cpp
new file mode 100644
--- /dev/null
+++ b/Engine/Source/Developer/MassEntityTestSuite/Private/Coverage/MassEntityViewCoverageTest.cpp
@@ -0,0 +1,94 @@
+// Copyright Epic Games, Inc. All Rights Reserved.
+
+#include "AITestsCommon.h"
+#include "MassEntityManager.h"
+#include "MassEntityTestTypes.h"
+#include "MassEntityView.h"
+#include "MassEntityTypes.h"
+
+#define LOCTEXT_NAMESPACE "MassTest"
+
+UE_DISABLE_OPTIMIZATION_SHIP
+
+namespace UE::Mass::Test::Coverage
+{
+
+struct FEntityView_TypeErasedFragmentAccess : FEntityTestBase
+{
+	virtual bool InstantTest() override
+	{
+		CA_ASSUME(EntityManager);
+
+		const FMassEntityHandle Entity = EntityManager->CreateEntity(IntsArchetype);
+		FMassEntityView View(*EntityManager, Entity);
+
+		// Typed access
+		FTestFragment_Int& TypedRef = View.GetFragmentData<FTestFragment_Int>();
+
+		// Type-erased access
+		FStructView StructView = View.GetFragmentDataStruct(FTestFragment_Int::StaticStruct());
+
+		AITEST_TRUE("Type-erased view is valid", StructView.IsValid());
+		AITEST_EQUAL("Type-erased pointer matches typed pointer",
+			StructView.GetMemory(), reinterpret_cast<const uint8*>(&TypedRef));
+
+		// Verify we can read through the type-erased view
+		FTestFragment_Int* ErasedPtr = StructView.GetPtr<FTestFragment_Int>();
+		AITEST_NOT_NULL("Can cast type-erased view back to typed", ErasedPtr);
+		AITEST_EQUAL("Values match", ErasedPtr->Value, TypedRef.Value);
+
+		return true;
+	}
+};
+IMPLEMENT_AI_INSTANT_TEST(FEntityView_TypeErasedFragmentAccess, "System.Mass.Coverage.EntityView.TypeErasedFragmentAccess");
+
+struct FEntityView_TypeErasedSharedAccess : FEntityTestBase
+{
+	virtual bool InstantTest() override
+	{
+		CA_ASSUME(EntityManager);
+
+		// Build archetype with shared and const shared fragments
+		FMassArchetypeCompositionDescriptor Composition;
+		Composition.Add<FTestFragment_Int>();
+		Composition.Add<FTestSharedFragment_Int>();
+		Composition.Add<FTestConstSharedFragment_Int>();
+
+		FMassArchetypeSharedFragmentValues SharedValues;
+		const FSharedStruct SharedInt = FSharedStruct::Make<FTestSharedFragment_Int>(42);
+		SharedValues.Add(SharedInt);
+		const FConstSharedStruct ConstSharedInt = FConstSharedStruct::Make<FTestConstSharedFragment_Int>(99);
+		SharedValues.Add(ConstSharedInt);
+		SharedValues.Sort();
+
+		// CreateArchetype takes only composition — shared values are passed to BatchCreateEntities
+		const FMassArchetypeHandle Archetype = EntityManager->CreateArchetype(Composition);
+		TArray<FMassEntityHandle> Entities;
+		EntityManager->BatchCreateEntities(Archetype, SharedValues, 1, Entities);
+		const FMassEntityHandle Entity = Entities[0];
+		FMassEntityView View(*EntityManager, Entity);
+
+		// Test type-erased const shared access
+		FConstStructView ConstSharedView = View.GetConstSharedFragmentDataStruct(FTestConstSharedFragment_Int::StaticStruct());
+		AITEST_TRUE("Const shared type-erased view is valid", ConstSharedView.IsValid());
+		const FTestConstSharedFragment_Int* ConstData = ConstSharedView.GetPtr<FTestConstSharedFragment_Int>();
+		AITEST_NOT_NULL("Const shared data is accessible", ConstData);
+		AITEST_EQUAL("Const shared value matches", ConstData->Value, 99);
+
+		// Test type-erased mutable shared access
+		FStructView SharedView = View.GetSharedFragmentDataStruct(FTestSharedFragment_Int::StaticStruct());
+		AITEST_TRUE("Shared type-erased view is valid", SharedView.IsValid());
+		FTestSharedFragment_Int* SharedData = SharedView.GetPtr<FTestSharedFragment_Int>();
+		AITEST_NOT_NULL("Shared data is accessible", SharedData);
+		AITEST_EQUAL("Shared value matches", SharedData->Value, 42);
+
+		return true;
+	}
+};
+IMPLEMENT_AI_INSTANT_TEST(FEntityView_TypeErasedSharedAccess, "System.Mass.Coverage.EntityView.TypeErasedSharedAccess");
+
+} // UE::Mass::Test::Coverage
+
+UE_ENABLE_OPTIMIZATION_SHIP
+
+#undef LOCTEXT_NAMESPACE
diff --git a/Engine/Source/Developer/MassEntityTestSuite/Private/Coverage/MassExecutionContextTest.cpp b/Engine/Source/Developer/MassEntityTestSuite/Private/Coverage/MassExecutionContextTest.cpp
new file mode 100644
--- /dev/null
+++ b/Engine/Source/Developer/MassEntityTestSuite/Private/Coverage/MassExecutionContextTest.cpp
@@ -0,0 +1,144 @@
+// Copyright Epic Games, Inc. All Rights Reserved.
+
+#include "AITestsCommon.h"
+#include "MassEntityManager.h"
+#include "MassEntityTestTypes.h"
+#include "MassExecutionContext.h"
+#include "MassEntityQuery.h"
+#include "MassObserverProcessor.h"
+#include "MassExecutor.h"
+#include "Engine/Engine.h"
+
+#define LOCTEXT_NAMESPACE "MassTest"
+
+UE_DISABLE_OPTIMIZATION_SHIP
+
+namespace UE::Mass::Test::Coverage
+{
+
+struct FExecutionContext_SubsystemAccessChecked : FEntityTestBase
+{
+	virtual bool InstantTest() override
+	{
+		CA_ASSUME(EntityManager);
+
+		const FMassEntityHandle Entity = EntityManager->CreateEntity(FloatsArchetype);
+
+		FMassEntityQuery Query(EntityManager);
+		Query.AddRequirement<FTestFragment_Float>(EMassFragmentAccess::ReadOnly);
+		Query.AddSubsystemRequirement<UMassTestWorldSubsystem>(EMassFragmentAccess::ReadWrite);
+		Query.AddSubsystemRequirement<UMassTestEngineSubsystem>(EMassFragmentAccess::ReadOnly);
+
+		bool bWorldSubsystemAccessSucceeded = false;
+		bool bEngineSubsystemAccessSucceeded = false;
+
+		FMassExecutionContext ExecutionContext = EntityManager->CreateExecutionContext(/*DeltaSeconds=*/0);
+		Query.ForEachEntityChunk(ExecutionContext, [&bWorldSubsystemAccessSucceeded, &bEngineSubsystemAccessSucceeded](FMassExecutionContext& Context)
+		{
+			UMassTestWorldSubsystem& WorldSubsystem = Context.GetMutableSubsystemChecked<UMassTestWorldSubsystem>();
+			bWorldSubsystemAccessSucceeded = true;
+
+			const UMassTestEngineSubsystem& EngineSubsystem = Context.GetSubsystemChecked<UMassTestEngineSubsystem>();
+			bEngineSubsystemAccessSucceeded = true;
+		});
+
+		AITEST_TRUE("World subsystem access via GetMutableSubsystemChecked succeeded", bWorldSubsystemAccessSucceeded);
+		AITEST_TRUE("Engine subsystem access via GetSubsystemChecked succeeded", bEngineSubsystemAccessSucceeded);
+
+		return true;
+	}
+};
+IMPLEMENT_AI_INSTANT_TEST(FExecutionContext_SubsystemAccessChecked, "System.Mass.Coverage.ExecutionContext.SubsystemAccessChecked");
+
+struct FExecutionContext_AuxData : FEntityTestBase
+{
+	virtual bool InstantTest() override
+	{
+		CA_ASSUME(EntityManager);
+
+		const FMassEntityHandle Entity = EntityManager->CreateEntity(IntsArchetype);
+
+		// Set up observer that checks AuxData
+		UMassTestProcessorBase* ObserverProcessor = NewTestProcessor<UMassTestProcessorBase>(EntityManager);
+		ObserverProcessor->EntityQuery.AddRequirement<FTestFragment_Int>(EMassFragmentAccess::ReadOnly);
+		ObserverProcessor->EntityQuery.AddTagRequirement<FTestTag_A>(EMassFragmentPresence::All);
+
+		bool bAuxDataValid = false;
+		EMassObservedOperation ObservedOp = EMassObservedOperation::MAX;
+		const UScriptStruct* ObservedType = nullptr;
+
+		ObserverProcessor->ForEachEntityChunkExecutionFunction = [&bAuxDataValid, &ObservedOp, &ObservedType](FMassExecutionContext& Context)
+		{
+			const FInstancedStruct& AuxData = Context.GetAuxData();
+			if (AuxData.IsValid())
+			{
+				const FMassObserverExecutionContext* ObserverContext = AuxData.GetPtr<FMassObserverExecutionContext>();
+				if (ObserverContext)
+				{
+					bAuxDataValid = true;
+					ObservedOp = ObserverContext->GetOperationType();
+					ObservedType = ObserverContext->GetCurrentType();
+				}
+			}
+		};
+
+		FMassObserverManager& ObserverManager = EntityManager->GetObserverManager();
+		ObserverManager.AddObserverInstance(FTestTag_A::StaticStruct(), EMassObservedOperationFlags::Add, ObserverProcessor);
+
+		// Trigger observer via deferred tag addition
+		EntityManager->Defer().AddTag<FTestTag_A>(Entity);
+		EntityManager->FlushCommands();
+
+		AITEST_TRUE("AuxData contained valid FMassObserverExecutionContext", bAuxDataValid);
+		AITEST_EQUAL("Observed operation was AddElement", ObservedOp, EMassObservedOperation::AddElement);
+		AITEST_EQUAL("Observed type was FTestTag_A", ObservedType, FTestTag_A::StaticStruct());
+
+		return true;
+	}
+};
+IMPLEMENT_AI_INSTANT_TEST(FExecutionContext_AuxData, "System.Mass.Coverage.ExecutionContext.AuxData");
+
+struct FExecutionContext_MidExecutionFlush : FEntityTestBase
+{
+	virtual bool InstantTest() override
+	{
+		CA_ASSUME(EntityManager);
+
+		// Create entity with Float fragment only
+		const FMassEntityHandle Entity = EntityManager->CreateEntity(FloatsArchetype);
+
+		FMassEntityQuery Query(EntityManager);
+		Query.AddRequirement<FTestFragment_Float>(EMassFragmentAccess::ReadOnly);
+
+		bool bFlushCalled = false;
+
+		bool bHasFragmentAfterFlush = false;
+		FMassExecutionContext ExecutionContext = EntityManager->CreateExecutionContext(/*DeltaSeconds=*/0);
+		Query.ForEachEntityChunk(ExecutionContext, [&bFlushCalled, &bHasFragmentAfterFlush, this](FMassExecutionContext& Context)
+		{
+			// Defer adding Int fragment to a different entity (not the one being iterated)
+			// and flush mid-execution
+			const FMassEntityHandle NewEntity = EntityManager->CreateEntity(EmptyArchetype);
+			Context.Defer().AddFragment<FTestFragment_Int>(NewEntity);
+			Context.FlushDeferred();
+			bFlushCalled = true;
+
+			// Check if the deferred command executed - new entity should now have FTestFragment_Int
+			bHasFragmentAfterFlush = EntityManager->GetArchetypeForEntity(NewEntity).IsValid()
+				&& FMassEntityView(*EntityManager, NewEntity).GetFragmentDataPtr<FTestFragment_Int>() != nullptr;
+		});
+
+		AITEST_TRUE("Deferred AddFragment executed after FlushDeferred", bHasFragmentAfterFlush);
+
+		AITEST_TRUE("FlushDeferred was called during iteration", bFlushCalled);
+
+		return true;
+	}
+};
+IMPLEMENT_AI_INSTANT_TEST(FExecutionContext_MidExecutionFlush, "System.Mass.Coverage.ExecutionContext.MidExecutionFlush");
+
+} // UE::Mass::Test::Coverage
+
+UE_ENABLE_OPTIMIZATION_SHIP
+
+#undef LOCTEXT_NAMESPACE
diff --git a/Engine/Source/Developer/MassEntityTestSuite/Private/Coverage/MassObserverProcessorTest.cpp b/Engine/Source/Developer/MassEntityTestSuite/Private/Coverage/MassObserverProcessorTest.cpp
new file mode 100644
--- /dev/null
+++ b/Engine/Source/Developer/MassEntityTestSuite/Private/Coverage/MassObserverProcessorTest.cpp
@@ -0,0 +1,105 @@
+// Copyright Epic Games, Inc. All Rights Reserved.
+
+#include "AITestsCommon.h"
+#include "MassEntityManager.h"
+#include "MassEntityTestTypes.h"
+#include "MassExecutionContext.h"
+#include "MassObserverProcessor.h"
+
+#define LOCTEXT_NAMESPACE "MassTest"
+
+UE_DISABLE_OPTIMIZATION_SHIP
+
+namespace UE::Mass::Test::Coverage
+{
+
+struct FObserver_BatchElementNotification : FEntityTestBase
+{
+	virtual bool InstantTest() override
+	{
+		CA_ASSUME(EntityManager);
+
+		// Create entities in an archetype without TagA
+		const int32 NumEntities = 50;
+		TArray<FMassEntityHandle> Entities;
+		EntityManager->BatchCreateEntities(IntsArchetype, NumEntities, Entities);
+
+		// Set up observer for FTestFragment_Float addition
+		UMassTestProcessorBase* ObserverProcessor = NewTestProcessor<UMassTestProcessorBase>(EntityManager);
+		ObserverProcessor->EntityQuery.AddRequirement<FTestFragment_Int>(EMassFragmentAccess::ReadOnly);
+		ObserverProcessor->EntityQuery.AddRequirement<FTestFragment_Float>(EMassFragmentAccess::ReadOnly);
+
+		int32 ObservedEntityCount = 0;
+		ObserverProcessor->ForEachEntityChunkExecutionFunction = [&ObservedEntityCount](FMassExecutionContext& Context)
+		{
+			ObservedEntityCount += Context.GetNumEntities();
+		};
+
+		FMassObserverManager& ObserverManager = EntityManager->GetObserverManager();
+		ObserverManager.AddObserverInstance(FTestFragment_Float::StaticStruct(), EMassObservedOperationFlags::Add, ObserverProcessor);
+
+		// Batch add element to all entities
+		EntityManager->BatchAddElementToEntities(Entities, FTestFragment_Float::StaticStruct());
+
+		AITEST_EQUAL("Observer processed all entities from batch add", ObservedEntityCount, NumEntities);
+
+		return true;
+	}
+};
+IMPLEMENT_AI_INSTANT_TEST(FObserver_BatchElementNotification, "System.Mass.Coverage.Observer.BatchElementNotification");
+
+struct FObserver_FlushOrdering : FEntityTestBase
+{
+	virtual bool InstantTest() override
+	{
+		CA_ASSUME(EntityManager);
+
+		const FMassEntityHandle Entity = EntityManager->CreateEntity(IntsArchetype);
+
+		// Track observer firing order
+		TArray<FString> FiringOrder;
+
+		// Observer for TagA addition
+		UMassTestProcessorBase* AddObserver = NewTestProcessor<UMassTestProcessorBase>(EntityManager);
+		AddObserver->EntityQuery.AddRequirement<FTestFragment_Int>(EMassFragmentAccess::ReadOnly);
+		AddObserver->EntityQuery.AddTagRequirement<FTestTag_A>(EMassFragmentPresence::All);
+		AddObserver->ForEachEntityChunkExecutionFunction = [&FiringOrder](FMassExecutionContext& Context)
+		{
+			FiringOrder.Add(TEXT("Add"));
+		};
+
+		// Observer for TagA removal
+		UMassTestProcessorBase* RemoveObserver = NewTestProcessor<UMassTestProcessorBase>(EntityManager);
+		RemoveObserver->EntityQuery.AddRequirement<FTestFragment_Int>(EMassFragmentAccess::ReadOnly);
+		RemoveObserver->ForEachEntityChunkExecutionFunction = [&FiringOrder](FMassExecutionContext& Context)
+		{
+			FiringOrder.Add(TEXT("Remove"));
+		};
+
+		FMassObserverManager& ObserverManager = EntityManager->GetObserverManager();
+		ObserverManager.AddObserverInstance(FTestTag_A::StaticStruct(), EMassObservedOperationFlags::Add, AddObserver);
+		ObserverManager.AddObserverInstance(FTestTag_A::StaticStruct(), EMassObservedOperationFlags::Remove, RemoveObserver);
+
+		// Queue add then remove in same flush
+		EntityManager->Defer().AddTag<FTestTag_A>(Entity);
+		EntityManager->Defer().RemoveTag<FTestTag_A>(Entity);
+		EntityManager->FlushCommands();
+
+		// Both observers should have fired in order: add then remove
+		AITEST_EQUAL("Both observers fired", FiringOrder.Num(), 2);
+		if (FiringOrder.Num() == 2)
+		{
+			AITEST_EQUAL("Add observer fired first", FiringOrder[0], TEXT("Add"));
+			AITEST_EQUAL("Remove observer fired second", FiringOrder[1], TEXT("Remove"));
+		}
+
+		return true;
+	}
+};
+IMPLEMENT_AI_INSTANT_TEST(FObserver_FlushOrdering, "System.Mass.Coverage.Observer.FlushOrdering");
+
+} // UE::Mass::Test::Coverage
+
+UE_ENABLE_OPTIMIZATION_SHIP
+
+#undef LOCTEXT_NAMESPACE
diff --git a/Engine/Source/Developer/MassEntityTestSuite/Private/Coverage/MassProcessorConfigTest.cpp b/Engine/Source/Developer/MassEntityTestSuite/Private/Coverage/MassProcessorConfigTest.cpp
new file mode 100644
--- /dev/null
+++ b/Engine/Source/Developer/MassEntityTestSuite/Private/Coverage/MassProcessorConfigTest.cpp
@@ -0,0 +1,84 @@
+// Copyright Epic Games, Inc. All Rights Reserved.
+
+#include "AITestsCommon.h"
+#include "MassEntityManager.h"
+#include "MassEntityTestTypes.h"
+#include "MassExecutionContext.h"
+#include "MassExecutor.h"
+#include "MassProcessingTypes.h"
+#include "MassProcessingContext.h"
+
+#define LOCTEXT_NAMESPACE "MassTest"
+
+UE_DISABLE_OPTIMIZATION_SHIP
+
+namespace UE::Mass::Test::Coverage
+{
+
+struct FProcessor_DynamicMarking : FEntityTestBase
+{
+	virtual bool InstantTest() override
+	{
+		CA_ASSUME(EntityManager);
+
+		UMassTestProcessorBase* Processor = NewTestProcessor<UMassTestProcessorBase>(EntityManager);
+		AITEST_FALSE("Processor is not dynamic by default", Processor->IsDynamic());
+
+		Processor->MarkAsDynamic();
+		AITEST_TRUE("Processor is dynamic after MarkAsDynamic", Processor->IsDynamic());
+
+		// Verify dynamic processor still executes
+		EntityManager->CreateEntity(FloatsArchetype);
+		Processor->EntityQuery.AddRequirement<FTestFragment_Float>(EMassFragmentAccess::ReadOnly);
+
+		int32 ExecutionCount = 0;
+		Processor->ForEachEntityChunkExecutionFunction = [&ExecutionCount](FMassExecutionContext& Context)
+		{
+			++ExecutionCount;
+		};
+
+		FMassProcessingContext ProcessingContext(*EntityManager, 0.f);
+		UE::Mass::Executor::Run(*Processor, ProcessingContext);
+
+		AITEST_EQUAL("Dynamic processor executed successfully", ExecutionCount, 1);
+
+		return true;
+	}
+};
+IMPLEMENT_AI_INSTANT_TEST(FProcessor_DynamicMarking, "System.Mass.Coverage.ProcessorConfig.DynamicMarking");
+
+struct FProcessor_QueryBasedPruning : FEntityTestBase
+{
+	virtual bool InstantTest() override
+	{
+		CA_ASSUME(EntityManager);
+
+		// Create entities with Float fragment only
+		EntityManager->CreateEntity(FloatsArchetype);
+
+		// Create processor that requires Bool fragment (not present in any entity)
+		UMassTestProcessorBase* Processor = NewTestProcessor<UMassTestProcessorBase>(EntityManager);
+		Processor->EntityQuery.AddRequirement<FTestFragment_Bool>(EMassFragmentAccess::ReadOnly);
+
+		int32 ExecutionCount = 0;
+		Processor->ForEachEntityChunkExecutionFunction = [&ExecutionCount](FMassExecutionContext& Context)
+		{
+			++ExecutionCount;
+		};
+
+		// Execute — processor should not execute because no entities match
+		FMassProcessingContext ProcessingContext(*EntityManager, 0.f);
+		UE::Mass::Executor::Run(*Processor, ProcessingContext);
+
+		AITEST_EQUAL("Processor with no matching entities did not execute", ExecutionCount, 0);
+
+		return true;
+	}
+};
+IMPLEMENT_AI_INSTANT_TEST(FProcessor_QueryBasedPruning, "System.Mass.Coverage.ProcessorConfig.QueryBasedPruning");
+
+} // UE::Mass::Test::Coverage
+
+UE_ENABLE_OPTIMIZATION_SHIP
+
+#undef LOCTEXT_NAMESPACE
diff --git a/Engine/Source/Developer/MassEntityTestSuite/Private/Coverage/MassProcessorCoverageTest.cpp b/Engine/Source/Developer/MassEntityTestSuite/Private/Coverage/MassProcessorCoverageTest.cpp
new file mode 100644
--- /dev/null
+++ b/Engine/Source/Developer/MassEntityTestSuite/Private/Coverage/MassProcessorCoverageTest.cpp
@@ -0,0 +1,159 @@
+// Copyright Epic Games, Inc. All Rights Reserved.
+
+#include "AITestsCommon.h"
+#include "MassEntityManager.h"
+#include "MassEntityTestTypes.h"
+#include "MassExecutionContext.h"
+#include "MassExecutor.h"
+#include "MassProcessingTypes.h"
+#include "MassProcessingContext.h"
+
+#define LOCTEXT_NAMESPACE "MassTest"
+
+UE_DISABLE_OPTIMIZATION_SHIP
+
+namespace UE::Mass::Test::Coverage
+{
+
+struct FProcessor_InitialState : FEntityTestBase
+{
+	virtual bool InstantTest() override
+	{
+		CA_ASSUME(EntityManager);
+
+		UMassTestProcessorBase* Processor = NewTestProcessor<UMassTestProcessorBase>(EntityManager);
+		AITEST_TRUE("Newly created processor is active by default", Processor->IsActive());
+
+		return true;
+	}
+};
+IMPLEMENT_AI_INSTANT_TEST(FProcessor_InitialState, "System.Mass.Coverage.Processor.InitialState");
+
+struct FProcessor_ActiveInactive : FEntityTestBase
+{
+	virtual bool InstantTest() override
+	{
+		CA_ASSUME(EntityManager);
+
+		EntityManager->CreateEntity(FloatsArchetype);
+
+		UMassTestProcessorBase* Processor = NewTestProcessor<UMassTestProcessorBase>(EntityManager);
+		Processor->EntityQuery.AddRequirement<FTestFragment_Float>(EMassFragmentAccess::ReadOnly);
+
+		int32 ExecutionCount = 0;
+		Processor->ForEachEntityChunkExecutionFunction = [&ExecutionCount](FMassExecutionContext& Context)
+		{
+			++ExecutionCount;
+		};
+
+		// Make inactive and execute
+		Processor->MakeInactive();
+		AITEST_FALSE("Processor is not active after MakeInactive", Processor->IsActive());
+
+		{
+			FMassProcessingContext ProcessingContext(*EntityManager, 0.f);
+			UE::Mass::Executor::Run(*Processor, ProcessingContext);
+		}
+		AITEST_EQUAL("Inactive processor did not execute", ExecutionCount, 0);
+
+		// Make active and execute
+		Processor->MakeActive();
+		AITEST_TRUE("Processor is active after MakeActive", Processor->IsActive());
+
+		{
+			FMassProcessingContext ProcessingContext(*EntityManager, 0.f);
+			UE::Mass::Executor::Run(*Processor, ProcessingContext);
+		}
+		AITEST_EQUAL("Active processor executed once", ExecutionCount, 1);
+
+		return true;
+	}
+};
+IMPLEMENT_AI_INSTANT_TEST(FProcessor_ActiveInactive, "System.Mass.Coverage.Processor.ActiveInactive");
+
+struct FProcessor_OneShot : FEntityTestBase
+{
+	virtual bool InstantTest() override
+	{
+		CA_ASSUME(EntityManager);
+
+		EntityManager->CreateEntity(FloatsArchetype);
+
+		UMassTestProcessorBase* Processor = NewTestProcessor<UMassTestProcessorBase>(EntityManager);
+		Processor->EntityQuery.AddRequirement<FTestFragment_Float>(EMassFragmentAccess::ReadOnly);
+
+		int32 ExecutionCount = 0;
+		Processor->ForEachEntityChunkExecutionFunction = [&ExecutionCount](FMassExecutionContext& Context)
+		{
+			++ExecutionCount;
+		};
+
+		Processor->MakeOneShot();
+		AITEST_TRUE("OneShot processor is active before first execution", Processor->IsActive());
+
+		// First execution — should run
+		{
+			FMassProcessingContext ProcessingContext(*EntityManager, 0.f);
+			UE::Mass::Executor::Run(*Processor, ProcessingContext);
+		}
+		AITEST_EQUAL("OneShot processor executed on first run", ExecutionCount, 1);
+		AITEST_FALSE("OneShot processor is inactive after first execution", Processor->IsActive());
+
+		// Second execution — should NOT run
+		{
+			FMassProcessingContext ProcessingContext(*EntityManager, 0.f);
+			UE::Mass::Executor::Run(*Processor, ProcessingContext);
+		}
+		AITEST_EQUAL("OneShot processor did not execute on second run", ExecutionCount, 1);
+
+		return true;
+	}
+};
+IMPLEMENT_AI_INSTANT_TEST(FProcessor_OneShot, "System.Mass.Coverage.Processor.OneShot");
+
+struct FProcessor_OneShotReactivation : FEntityTestBase
+{
+	virtual bool InstantTest() override
+	{
+		CA_ASSUME(EntityManager);
+
+		EntityManager->CreateEntity(FloatsArchetype);
+
+		UMassTestProcessorBase* Processor = NewTestProcessor<UMassTestProcessorBase>(EntityManager);
+		Processor->EntityQuery.AddRequirement<FTestFragment_Float>(EMassFragmentAccess::ReadOnly);
+
+		int32 ExecutionCount = 0;
+		Processor->ForEachEntityChunkExecutionFunction = [&ExecutionCount](FMassExecutionContext& Context)
+		{
+			++ExecutionCount;
+		};
+
+		// OneShot, execute, auto-deactivate
+		Processor->MakeOneShot();
+		{
+			FMassProcessingContext ProcessingContext(*EntityManager, 0.f);
+			UE::Mass::Executor::Run(*Processor, ProcessingContext);
+		}
+		AITEST_EQUAL("Executed once", ExecutionCount, 1);
+		AITEST_FALSE("Inactive after oneshot", Processor->IsActive());
+
+		// Reactivate
+		Processor->MakeActive();
+		AITEST_TRUE("Active after reactivation", Processor->IsActive());
+
+		{
+			FMassProcessingContext ProcessingContext(*EntityManager, 0.f);
+			UE::Mass::Executor::Run(*Processor, ProcessingContext);
+		}
+		AITEST_EQUAL("Executed again after reactivation", ExecutionCount, 2);
+
+		return true;
+	}
+};
+IMPLEMENT_AI_INSTANT_TEST(FProcessor_OneShotReactivation, "System.Mass.Coverage.Processor.OneShotReactivation");
+
+} // UE::Mass::Test::Coverage
+
+UE_ENABLE_OPTIMIZATION_SHIP
+
+#undef LOCTEXT_NAMESPACE
diff --git a/Engine/Source/Developer/MassEntityTestSuite/Private/Coverage/MassRelationCoverageTest.cpp b/Engine/Source/Developer/MassEntityTestSuite/Private/Coverage/MassRelationCoverageTest.cpp
new file mode 100644
--- /dev/null
+++ b/Engine/Source/Developer/MassEntityTestSuite/Private/Coverage/MassRelationCoverageTest.cpp
@@ -0,0 +1,152 @@
+// Copyright Epic Games, Inc. All Rights Reserved.
+
+#include "AITestsCommon.h"
+#include "MassEntityBuilder.h"
+#include "MassRelationManager.h"
+#include "MassEntityTestTypes.h"
+
+#define LOCTEXT_NAMESPACE "MassTest"
+
+UE_DISABLE_OPTIMIZATION_SHIP
+
+namespace UE::Mass::Test::Coverage
+{
+
+struct FCoverageRelationTestBase : FEntityTestBase
+{
+	UScriptStruct* RelationTypeA = nullptr;
+	UScriptStruct* RelationTypeB = nullptr;
+	UE::Mass::FTypeHandle RelationHandleA;
+	UE::Mass::FTypeHandle RelationHandleB;
+	UE::Mass::FRelationManager* RelationManagerPtr = nullptr;
+
+	static TNotNull<UScriptStruct*> MakeRelationType(FName TypeName)
+	{
+		UScriptStruct* NewStruct = NewObject<UScriptStruct>(UScriptStruct::StaticClass(),
+			FName(*FString::Printf(TEXT("CoverageRelation_%s"), *TypeName.ToString())), RF_Public);
+		NewStruct->SetSuperStruct(FMassRelation::StaticStruct());
+		return NewStruct;
+	}
+
+	static UE::Mass::FRelationTypeTraits CreateRelationTraits(TNotNull<UScriptStruct*> RelationType)
+	{
+		UE::Mass::FRelationTypeTraits Traits(RelationType);
+		Traits.RoleTraits[static_cast<uint8>(UE::Mass::ERelationRole::Object)].DestructionPolicy = UE::Mass::ERemovalPolicy::CleanUp;
+		Traits.RoleTraits[static_cast<uint8>(UE::Mass::ERelationRole::Subject)].DestructionPolicy = UE::Mass::ERemovalPolicy::CleanUp;
+		return Traits;
+	}
+
+	virtual bool SetUp() override
+	{
+		if (FEntityTestBase::SetUp())
+		{
+			RelationTypeA = MakeRelationType(TEXT("TypeA"));
+			RelationHandleA = EntityManager->GetTypeManager().RegisterType(CreateRelationTraits(RelationTypeA));
+
+			RelationTypeB = MakeRelationType(TEXT("TypeB"));
+			RelationHandleB = EntityManager->GetTypeManager().RegisterType(CreateRelationTraits(RelationTypeB));
+
+			RelationManagerPtr = &EntityManager->GetRelationManager();
+		}
+		return RelationHandleA.IsValid() && RelationHandleB.IsValid();
+	}
+};
+
+struct FRelation_RoleEntities : FCoverageRelationTestBase
+{
+	virtual bool InstantTest() override
+	{
+		CA_ASSUME(EntityManager);
+
+		const FMassEntityHandle SubjectEntity = EntityManager->CreateEntity(IntsArchetype);
+		const FMassEntityHandle ObjectEntity = EntityManager->CreateEntity(FloatsArchetype);
+
+		const FMassEntityHandle RelationEntity = RelationManagerPtr->CreateRelationInstance(RelationHandleA, SubjectEntity, ObjectEntity);
+		EntityManager->FlushCommands();
+
+		AITEST_TRUE("Relation entity created", RelationEntity.IsValid());
+
+		// Query subjects of the object
+		TArray<FMassEntityHandle> Subjects = RelationManagerPtr->GetRelationSubjects(RelationHandleA, ObjectEntity);
+		AITEST_EQUAL("One subject found for ObjectEntity", Subjects.Num(), 1);
+		if (Subjects.Num() > 0)
+		{
+			AITEST_EQUAL("Subject is correct entity", Subjects[0], SubjectEntity);
+		}
+
+		// Query objects of the subject
+		TArray<FMassEntityHandle> Objects = RelationManagerPtr->GetRelationObjects(RelationHandleA, SubjectEntity);
+		AITEST_EQUAL("One object found for SubjectEntity", Objects.Num(), 1);
+		if (Objects.Num() > 0)
+		{
+			AITEST_EQUAL("Object is correct entity", Objects[0], ObjectEntity);
+		}
+
+		return true;
+	}
+};
+IMPLEMENT_AI_INSTANT_TEST(FRelation_RoleEntities, "System.Mass.Coverage.Relation.RoleEntities");
+
+struct FRelation_MultiType : FCoverageRelationTestBase
+{
+	virtual bool InstantTest() override
+	{
+		CA_ASSUME(EntityManager);
+
+		const FMassEntityHandle EntityA = EntityManager->CreateEntity(IntsArchetype);
+		const FMassEntityHandle EntityB = EntityManager->CreateEntity(FloatsArchetype);
+		const FMassEntityHandle EntityC = EntityManager->CreateEntity(IntsArchetype);
+
+		// Create two different relation types from EntityA
+		RelationManagerPtr->CreateRelationInstance(RelationHandleA, EntityA, EntityB);
+		RelationManagerPtr->CreateRelationInstance(RelationHandleB, EntityA, EntityC);
+		EntityManager->FlushCommands();
+
+		// Query each type independently
+		TArray<FMassEntityHandle> ObjectsA = RelationManagerPtr->GetRelationObjects(RelationHandleA, EntityA);
+		AITEST_EQUAL("One TypeA relation object", ObjectsA.Num(), 1);
+		if (ObjectsA.Num() > 0)
+		{
+			AITEST_EQUAL("TypeA object is EntityB", ObjectsA[0], EntityB);
+		}
+
+		TArray<FMassEntityHandle> ObjectsB = RelationManagerPtr->GetRelationObjects(RelationHandleB, EntityA);
+		AITEST_EQUAL("One TypeB relation object", ObjectsB.Num(), 1);
+		if (ObjectsB.Num() > 0)
+		{
+			AITEST_EQUAL("TypeB object is EntityC", ObjectsB[0], EntityC);
+		}
+
+		return true;
+	}
+};
+IMPLEMENT_AI_INSTANT_TEST(FRelation_MultiType, "System.Mass.Coverage.Relation.MultiType");
+
+struct FRelation_DataAccess : FCoverageRelationTestBase
+{
+	virtual bool InstantTest() override
+	{
+		CA_ASSUME(EntityManager);
+
+		const FMassEntityHandle SubjectEntity = EntityManager->CreateEntity(IntsArchetype);
+		const FMassEntityHandle ObjectEntity = EntityManager->CreateEntity(FloatsArchetype);
+
+		RelationManagerPtr->CreateRelationInstance(RelationHandleA, SubjectEntity, ObjectEntity);
+		EntityManager->FlushCommands();
+
+		const bool bIsSubject = RelationManagerPtr->IsSubjectOfRelation(RelationHandleA, SubjectEntity, ObjectEntity);
+		AITEST_TRUE("SubjectEntity is subject of relation to ObjectEntity", bIsSubject);
+
+		const bool bIsNotSubject = RelationManagerPtr->IsSubjectOfRelation(RelationHandleA, ObjectEntity, SubjectEntity);
+		AITEST_FALSE("ObjectEntity is NOT subject of relation to SubjectEntity", bIsNotSubject);
+
+		return true;
+	}
+};
+IMPLEMENT_AI_INSTANT_TEST(FRelation_DataAccess, "System.Mass.Coverage.Relation.DataAccess");
+
+} // UE::Mass::Test::Coverage
+
+UE_ENABLE_OPTIMIZATION_SHIP
+
+#undef LOCTEXT_NAMESPACE
diff --git a/Engine/Source/Developer/MassEntityTestSuite/Private/Coverage/MassSerializationCoverageTest.cpp b/Engine/Source/Developer/MassEntityTestSuite/Private/Coverage/MassSerializationCoverageTest.cpp
new file mode 100644
--- /dev/null
+++ b/Engine/Source/Developer/MassEntityTestSuite/Private/Coverage/MassSerializationCoverageTest.cpp
@@ -0,0 +1,46 @@
+// Copyright Epic Games, Inc. All Rights Reserved.
+
+#include "AITestsCommon.h"
+#include "MassEntityManager.h"
+#include "MassEntityTestTypes.h"
+#include "MassEntityHandle.h"
+
+#define LOCTEXT_NAMESPACE "MassTest"
+
+UE_DISABLE_OPTIMIZATION_SHIP
+
+namespace UE::Mass::Test::Coverage
+{
+
+struct FSerialization_EntityHandleRoundTrip : FEntityTestBase
+{
+	virtual bool InstantTest() override
+	{
+		CA_ASSUME(EntityManager);
+
+		const FMassEntityHandle Original = EntityManager->CreateEntity(IntsArchetype);
+		AITEST_TRUE("Original handle is valid", Original.IsSet());
+
+		// Round-trip through AsNumber/FromNumber
+		const uint64 Number = Original.AsNumber();
+		const FMassEntityHandle Reconstructed = FMassEntityHandle::FromNumber(Number);
+
+		AITEST_EQUAL("Reconstructed handle equals original", Reconstructed, Original);
+		AITEST_TRUE("Reconstructed handle is valid", EntityManager->IsEntityValid(Reconstructed));
+
+		// Verify they reference the same entity data
+		FTestFragment_Int& OriginalData = EntityManager->GetFragmentDataChecked<FTestFragment_Int>(Original);
+		OriginalData.Value = 777;
+		const FTestFragment_Int& ReconstructedData = EntityManager->GetFragmentDataChecked<FTestFragment_Int>(Reconstructed);
+		AITEST_EQUAL("Both handles reference the same entity data", ReconstructedData.Value, 777);
+
+		return true;
+	}
+};
+IMPLEMENT_AI_INSTANT_TEST(FSerialization_EntityHandleRoundTrip, "System.Mass.Coverage.Serialization.EntityHandleRoundTrip");
+
+} // UE::Mass::Test::Coverage
+
+UE_ENABLE_OPTIMIZATION_SHIP
+
+#undef LOCTEXT_NAMESPACE
diff --git a/Engine/Source/Developer/MassEntityTestSuite/Private/Coverage/MassSharedFragmentValuesTest.cpp b/Engine/Source/Developer/MassEntityTestSuite/Private/Coverage/MassSharedFragmentValuesTest.cpp
new file mode 100644
--- /dev/null
+++ b/Engine/Source/Developer/MassEntityTestSuite/Private/Coverage/MassSharedFragmentValuesTest.cpp
@@ -0,0 +1,114 @@
+// Copyright Epic Games, Inc. All Rights Reserved.
+
+#include "AITestsCommon.h"
+#include "MassEntityManager.h"
+#include "MassEntityTestTypes.h"
+#include "MassEntityTypes.h"
+
+#define LOCTEXT_NAMESPACE "MassTest"
+
+UE_DISABLE_OPTIMIZATION_SHIP
+
+namespace UE::Mass::Test::Coverage
+{
+
+struct FSharedValues_Replace : FEntityTestBase
+{
+	virtual bool InstantTest() override
+	{
+		CA_ASSUME(EntityManager);
+
+		FMassArchetypeSharedFragmentValues SharedValues;
+		const FSharedStruct OriginalShared = FSharedStruct::Make<FTestSharedFragment_Int>(10);
+		SharedValues.Add(OriginalShared);
+
+		// Verify original value
+		const FTestSharedFragment_Int* Data = SharedValues.GetSharedFragmentStruct(FTestSharedFragment_Int::StaticStruct()).GetPtr<FTestSharedFragment_Int>();
+		AITEST_NOT_NULL("Original shared data accessible", Data);
+		AITEST_EQUAL("Original value is 10", Data->Value, 10);
+
+		// Replace with new value
+		const FSharedStruct ReplacementShared = FSharedStruct::Make<FTestSharedFragment_Int>(42);
+		SharedValues.ReplaceSharedFragments(MakeArrayView(&ReplacementShared, 1));
+
+		// Verify replaced value
+		const FTestSharedFragment_Int* ReplacedData = SharedValues.GetSharedFragmentStruct(FTestSharedFragment_Int::StaticStruct()).GetPtr<FTestSharedFragment_Int>();
+		AITEST_NOT_NULL("Replaced shared data accessible", ReplacedData);
+		AITEST_EQUAL("Replaced value is 42", ReplacedData->Value, 42);
+
+		return true;
+	}
+};
+IMPLEMENT_AI_INSTANT_TEST(FSharedValues_Replace, "System.Mass.Coverage.SharedFragmentValues.Replace");
+
+struct FSharedValues_DoesMatchComposition : FEntityTestBase
+{
+	virtual bool InstantTest() override
+	{
+		FMassArchetypeSharedFragmentValues SharedValues;
+		const FSharedStruct SharedInt = FSharedStruct::Make<FTestSharedFragment_Int>(1);
+		SharedValues.Add(SharedInt);
+		const FConstSharedStruct ConstSharedInt = FConstSharedStruct::Make<FTestConstSharedFragment_Int>(2);
+		SharedValues.Add(ConstSharedInt);
+
+		// Matching composition
+		FMassArchetypeCompositionDescriptor MatchingDesc;
+		MatchingDesc.Add<FTestSharedFragment_Int>();
+		MatchingDesc.Add<FTestConstSharedFragment_Int>();
+
+		AITEST_TRUE("Matches composition with same shared types", SharedValues.DoesMatchComposition(MatchingDesc));
+
+		// Non-matching: missing a type
+		FMassArchetypeCompositionDescriptor PartialDesc;
+		PartialDesc.Add<FTestSharedFragment_Int>();
+		// Missing FTestConstSharedFragment_Int
+
+		AITEST_FALSE("Does not match composition with missing type", SharedValues.DoesMatchComposition(PartialDesc));
+
+		return true;
+	}
+};
+IMPLEMENT_AI_INSTANT_TEST(FSharedValues_DoesMatchComposition, "System.Mass.Coverage.SharedFragmentValues.DoesMatchComposition");
+
+struct FSharedValues_CreateCombined : FEntityTestBase
+{
+	virtual bool InstantTest() override
+	{
+		// Create original shared values with one shared fragment
+		FMassArchetypeSharedFragmentValues Original;
+		const FConstSharedStruct ConstSharedInt = FConstSharedStruct::Make<FTestConstSharedFragment_Int>(10);
+		Original.Add(ConstSharedInt);
+
+		// Modification adds a new shared fragment
+		FMassArchetypeSharedFragmentValues Modification;
+		const FConstSharedStruct ConstSharedFloat = FConstSharedStruct::Make<FTestConstSharedFragment_Float>(3.14f);
+		Modification.Add(ConstSharedFloat);
+
+		// New composition includes both original and new types
+		FMassElementBitSet NewComposition;
+		NewComposition.Add(FTestConstSharedFragment_Int::StaticStruct());
+		NewComposition.Add(FTestConstSharedFragment_Float::StaticStruct());
+
+		FMassArchetypeSharedFragmentValues Combined = FMassArchetypeSharedFragmentValues::CreateCombined(
+			Original, NewComposition, &Modification);
+
+		// Verify combined has both fragments
+		const FTestConstSharedFragment_Int* IntData = Combined.GetConstSharedFragmentStruct(
+			FTestConstSharedFragment_Int::StaticStruct()).GetPtr<FTestConstSharedFragment_Int>();
+		AITEST_NOT_NULL("Combined has original const shared int", IntData);
+		AITEST_EQUAL("Original value preserved", IntData->Value, 10);
+
+		const FTestConstSharedFragment_Float* FloatData = Combined.GetConstSharedFragmentStruct(
+			FTestConstSharedFragment_Float::StaticStruct()).GetPtr<FTestConstSharedFragment_Float>();
+		AITEST_NOT_NULL("Combined has added const shared float", FloatData);
+
+		return true;
+	}
+};
+IMPLEMENT_AI_INSTANT_TEST(FSharedValues_CreateCombined, "System.Mass.Coverage.SharedFragmentValues.CreateCombined");
+
+} // UE::Mass::Test::Coverage
+
+UE_ENABLE_OPTIMIZATION_SHIP
+
+#undef LOCTEXT_NAMESPACE
diff --git a/Engine/Source/Runtime/MassEntity/Public/MassArchetypeTypes.h b/Engine/Source/Runtime/MassEntity/Public/MassArchetypeTypes.h
--- a/Engine/Source/Runtime/MassEntity/Public/MassArchetypeTypes.h
+++ b/Engine/Source/Runtime/MassEntity/Public/MassArchetypeTypes.h
@@ -67,7 +67,7 @@ private:
 struct FMassArchetypeVersionedHandle final
 {
 	FMassArchetypeVersionedHandle() = default;
-	FMassArchetypeVersionedHandle(const FMassArchetypeHandle& InHandle);
+	MASSENTITY_API FMassArchetypeVersionedHandle(const FMassArchetypeHandle& InHandle);
 	FMassArchetypeVersionedHandle(FMassArchetypeHandle&& InHandle);
 
 	bool IsValid() const;
diff --git a/Engine/Source/Runtime/MassEntity/Public/MassEntityTypes.h b/Engine/Source/Runtime/MassEntity/Public/MassEntityTypes.h
--- a/Engine/Source/Runtime/MassEntity/Public/MassEntityTypes.h
+++ b/Engine/Source/Runtime/MassEntity/Public/MassEntityTypes.h
@@ -325,10 +325,16 @@ public:
 };
 PRAGMA_ENABLE_DEPRECATION_WARNINGS
 
-/** 
+/**
  * Wrapper for const and non-const shared fragment containers that tracks which struct types it holds (via a FMassSharedFragmentBitSet).
- * Note that having multiple instanced of a given struct type is not supported and Add* functions will fetch the previously 
+ * Note that having multiple instances of a given struct type is not supported and Add* functions will fetch the previously
  * added fragment instead of adding a new one.
+ *
+ * Shared fragment TYPES contribute to archetype composition (two archetypes with different shared fragment type sets are distinct).
+ * Shared fragment VALUES are stored per-chunk — entities within the same archetype but with different shared values are placed
+ * in separate chunks. Chunk selection uses IsEquivalent() which compares by pointer identity (hash of backing struct memory
+ * addresses). When shared fragments are obtained via GetOrCreateConstSharedFragment/GetOrCreateSharedFragment, identical
+ * values are deduplicated (same pointer), so IsEquivalent effectively acts as value comparison in the expected usage path.
  */
 struct FMassArchetypeSharedFragmentValues
 {
@@ -387,7 +393,7 @@ struct FMassArchetypeSharedFragmentValues
 	 *	Needs to be consistent with NewSharedElementComposition, meaning, not add anything that's not already in the bitset. 
 	 * Also, if NewSharedElementComposition contains elements not in Original, then SharedFragmentValuesModification has to provide a value for it.  
 	 */
-	static FMassArchetypeSharedFragmentValues CreateCombined(const FMassArchetypeSharedFragmentValues& Original
+	MASSENTITY_API static FMassArchetypeSharedFragmentValues CreateCombined(const FMassArchetypeSharedFragmentValues& Original
 		, const FMassElementBitSet& NewSharedElementComposition, const FMassArchetypeSharedFragmentValues* SharedFragmentValuesModification);
 
 	/**
