[Mass] Issue 1273382 fix, using debug-only functions in non-debug code

#rnx

#Synced-CL 52056998

diff --git a/Engine/Source/Developer/MassEntityTestSuite/Private/Coverage/MassEntityManagerTest.cpp b/Engine/Source/Developer/MassEntityTestSuite/Private/Coverage/MassEntityManagerTest.cpp
--- a/Engine/Source/Developer/MassEntityTestSuite/Private/Coverage/MassEntityManagerTest.cpp
+++ b/Engine/Source/Developer/MassEntityTestSuite/Private/Coverage/MassEntityManagerTest.cpp
@@ -14,6 +14,7 @@ UE_DISABLE_OPTIMIZATION_SHIP
 namespace UE::Mass::Test::Coverage
 {
 
+#if WITH_MASSENTITY_DEBUG
 struct FManager_EntityCompaction : FEntityTestBase
 {
 	virtual bool InstantTest() override
@@ -64,6 +65,7 @@ struct FManager_EntityCompaction : FEntityTestBase
 	}
 };
 IMPLEMENT_AI_INSTANT_TEST(FManager_EntityCompaction, "System.Mass.Coverage.Manager.EntityCompaction");
+#endif // WITH_MASSENTITY_DEBUG
 
 struct FManager_SwapTagsForEntity : FEntityTestBase
 {
diff --git a/Engine/Source/Developer/MassEntityTestSuite/Private/Coverage/MassEntityQueryTest.cpp b/Engine/Source/Developer/MassEntityTestSuite/Private/Coverage/MassEntityQueryTest.cpp
--- a/Engine/Source/Developer/MassEntityTestSuite/Private/Coverage/MassEntityQueryTest.cpp
+++ b/Engine/Source/Developer/MassEntityTestSuite/Private/Coverage/MassEntityQueryTest.cpp
@@ -13,7 +13,7 @@ UE_DISABLE_OPTIMIZATION_SHIP
 
 namespace UE::Mass::Test::Coverage
 {
-
+#if WITH_MASSENTITY_DEBUG
 struct FQuery_ChunkFilter_Basic : FEntityTestBase
 {
 	virtual bool InstantTest() override
@@ -202,7 +202,7 @@ struct FQuery_ChunkFilter_WithCollections : FEntityTestBase
 	}
 };
 IMPLEMENT_AI_INSTANT_TEST(FQuery_ChunkFilter_WithCollections, "System.Mass.Coverage.Query.ChunkFilter.WithCollections");
-
+#endif // WITH_MASSENTITY_DEBUG
 } // UE::Mass::Test::Coverage
 
 UE_ENABLE_OPTIMIZATION_SHIP
