Fix shared fragment values being dropped during RemoveFragmentListFromEntity archetype transitions
Add typed template overload for RemoveConstSharedFragmentFromEntity
Add tests covering shared fragment value preservation across archetype moves and the new typed removal API

[RN] Make the mass RemoveFragmentList remove correctly the SharedFragments if provided in the list.
#Synced-CL 51802337

diff --git a/Engine/Source/Developer/MassEntityTestSuite/Private/MassSharedFragmentsTest.cpp b/Engine/Source/Developer/MassEntityTestSuite/Private/MassSharedFragmentsTest.cpp
--- a/Engine/Source/Developer/MassEntityTestSuite/Private/MassSharedFragmentsTest.cpp
+++ b/Engine/Source/Developer/MassEntityTestSuite/Private/MassSharedFragmentsTest.cpp
@@ -671,6 +671,73 @@ struct FSharedFragment_RemoveFromEntity : FEntityTestBase
 };
 IMPLEMENT_AI_INSTANT_TEST(FSharedFragment_RemoveFromEntity, "System.Mass.SharedFragments.RemoveFromEntity");
 
+// Tests that an entity's shared fragment values are preserved (not dropped) when
+// RemoveFragmentListFromEntity triggers an archetype change.
+struct FRemoveFragmentList_PreservesSharedFragmentValues : FEntityTestBase
+{
+	virtual bool InstantTest() override
+	{
+		constexpr int32 TestIntValue = 42;
+		const FTestSharedFragment_Int FragmentInstance(TestIntValue);
+		const FSharedStruct SharedFragmentInstance = EntityManager->GetOrCreateSharedFragment(FragmentInstance);
+
+		// Start with an entity that has a float fragment (FloatsArchetype)
+		const FMassEntityHandle EntityHandle = EntityManager->CreateEntity(FloatsArchetype);
+
+		// Add a shared fragment — entity now lives in a new archetype: {float, SharedInt}
+		AITEST_TRUE("Adding shared fragment should succeed",
+			EntityManager->AddSharedFragmentToEntity(EntityHandle, SharedFragmentInstance));
+
+		const FTestSharedFragment_Int* EntitySharedFragment = EntityManager->GetSharedFragmentDataPtr<FTestSharedFragment_Int>(EntityHandle);
+		AITEST_NOT_NULL("Entity should have the shared fragment before the remove", EntitySharedFragment);
+		AITEST_EQUAL("Shared fragment should hold the configured value", EntitySharedFragment->Value, TestIntValue);
+
+		// Remove the regular float fragment via RemoveFragmentListFromEntity — triggers an archetype change to {SharedInt}
+		const UScriptStruct* FragmentsToRemove[] = { FTestFragment_Float::StaticStruct() };
+		EntityManager->RemoveFragmentListFromEntity(EntityHandle, MakeArrayView(FragmentsToRemove));
+
+		// Shared fragment must survive the archetype move
+		EntitySharedFragment = EntityManager->GetSharedFragmentDataPtr<FTestSharedFragment_Int>(EntityHandle);
+		AITEST_NOT_NULL("Entity should still have the shared fragment after removing the regular fragment", EntitySharedFragment);
+		AITEST_EQUAL("Shared fragment value should be unchanged after the archetype move", EntitySharedFragment->Value, TestIntValue);
+
+		return true;
+	}
+};
+IMPLEMENT_AI_INSTANT_TEST(FRemoveFragmentList_PreservesSharedFragmentValues, "System.Mass.SharedFragments.RemoveFragmentList.PreservesSharedFragmentValues");
+
+// Tests the typed template overload RemoveConstSharedFragmentFromEntity<T>().
+struct FConstSharedFragment_RemoveFromEntity_Typed : FEntityTestBase
+{
+	virtual bool InstantTest() override
+	{
+		constexpr int32 TestIntValue = 77;
+		const FTestConstSharedFragment_Int FragmentInstance(TestIntValue);
+		const FConstSharedStruct SharedFragmentInstance = EntityManager->GetOrCreateConstSharedFragment(FragmentInstance);
+
+		const FMassEntityHandle EntityHandle = EntityManager->CreateEntity(FloatsArchetype);
+
+		AITEST_TRUE("Adding const shared fragment should succeed",
+			EntityManager->AddConstSharedFragmentToEntity(EntityHandle, SharedFragmentInstance));
+
+		const FTestConstSharedFragment_Int* EntitySharedFragment = EntityManager->GetConstSharedFragmentDataPtr<FTestConstSharedFragment_Int>(EntityHandle);
+		AITEST_NOT_NULL("Entity should have the const shared fragment", EntitySharedFragment);
+		AITEST_EQUAL("Const shared fragment should hold the configured value", EntitySharedFragment->Value, TestIntValue);
+
+		// Use the typed template overload (the new API being tested)
+		AITEST_TRUE("Typed RemoveConstSharedFragmentFromEntity<T>() should succeed",
+			EntityManager->RemoveConstSharedFragmentFromEntity<FTestConstSharedFragment_Int>(EntityHandle));
+
+		EntitySharedFragment = EntityManager->GetConstSharedFragmentDataPtr<FTestConstSharedFragment_Int>(EntityHandle);
+		AITEST_NULL("Entity should no longer have the const shared fragment after typed removal", EntitySharedFragment);
+
+		AITEST_EQUAL("Entity's archetype should be back to the initial one", EntityManager->GetArchetypeForEntity(EntityHandle), FloatsArchetype);
+
+		return true;
+	}
+};
+IMPLEMENT_AI_INSTANT_TEST(FConstSharedFragment_RemoveFromEntity_Typed, "System.Mass.ConstSharedFragments.RemoveFromEntity.Typed");
+
 struct FSharedFragment_BatchAddToEntity : FEntityTestBase
 {
 	virtual bool InstantTest() override
diff --git a/Engine/Source/Runtime/MassEntity/Private/MassEntityManager.cpp b/Engine/Source/Runtime/MassEntity/Private/MassEntityManager.cpp
--- a/Engine/Source/Runtime/MassEntity/Private/MassEntityManager.cpp
+++ b/Engine/Source/Runtime/MassEntity/Private/MassEntityManager.cpp
@@ -1322,10 +1322,16 @@ void FMassEntityManager::RemoveFragmentListFromEntity(FMassEntityHandle EntityHa
 		// Find overlap.  It isn't guaranteed that the old archetype has all the fragments being removed.
 		ObserverManager.OnPreCompositionRemoved(EntityHandle, OldArchetype->GetCompositionBitSet().GetOverlap(FragmentsToRemove));
 
-		// Move the entity over
+		// Move the entity over and its new shared fragments
 		FMassArchetypeData& NewArchetype = FMassArchetypeHelper::ArchetypeDataFromHandleChecked(NewArchetypeHandle);
 		NewArchetype.CopyDebugNamesFrom(*OldArchetype);
-		OldArchetype->MoveEntityToAnotherArchetype(EntityHandle, NewArchetype);
+		const FMassArchetypeSharedFragmentValues& OldSharedFragmentValues = OldArchetype->GetSharedFragmentValues(EntityHandle.Index);
+		FMassArchetypeSharedFragmentValues NewSharedFragmentValues(OldSharedFragmentValues);
+		NewSharedFragmentValues.Remove(FragmentsToRemove.GetSharedFragments());
+		// @Todo, optimization, sorting should not be required when only removing elements.
+		NewSharedFragmentValues.Sort();
+
+		OldArchetype->MoveEntityToAnotherArchetype(EntityHandle, NewArchetype, &NewSharedFragmentValues);
 
 		GetEntityStorageInterface().SetArchetypeFromShared(EntityHandle.Index, NewArchetypeHandle.DataPtr);
 	}
diff --git a/Engine/Source/Runtime/MassEntity/Public/MassEntityManager.h b/Engine/Source/Runtime/MassEntity/Public/MassEntityManager.h
--- a/Engine/Source/Runtime/MassEntity/Public/MassEntityManager.h
+++ b/Engine/Source/Runtime/MassEntity/Public/MassEntityManager.h
@@ -522,6 +522,12 @@ public:
 	 */
 	UE_API bool RemoveConstSharedFragmentFromEntity(const FMassEntityHandle EntityHandle, TNotNull<const UScriptStruct*> ConstSharedFragmentType);
 
+	template<UE::Mass::CConstSharedFragment T>
+	bool RemoveConstSharedFragmentFromEntity(const FMassEntityHandle EntityHandle)
+	{
+		return RemoveConstSharedFragmentFromEntity(EntityHandle, T::StaticStruct());
+	}
+
 	UE_DEPRECATED(5.8, "Use overload taking TNotNull")
 	bool RemoveConstSharedFragmentFromEntity(const FMassEntityHandle EntityHandle, const UScriptStruct& ConstSharedFragmentType)
 	{
@@ -546,9 +552,9 @@ public:
 	UE_API bool RemoveSharedFragmentFromEntity(const FMassEntityHandle EntityHandle, TNotNull<const UScriptStruct*> SharedFragmentType);
 
 	UE_DEPRECATED(5.8, "Use overload taking TNotNull")
-	bool RemoveSharedFragmentFromEntity(const FMassEntityHandle EntityHandle, const UScriptStruct& ConstSharedFragmentType)
+	bool RemoveSharedFragmentFromEntity(const FMassEntityHandle EntityHandle, const UScriptStruct& SharedFragmentType)
 	{
-		return RemoveSharedFragmentFromEntity(EntityHandle, &ConstSharedFragmentType);
+		return RemoveSharedFragmentFromEntity(EntityHandle, &SharedFragmentType);
 	}
 
 	/**
