[Mass] removing MassCore's and MassEntity's dependency on AITestSuite. The testable-ensure mechanism (global counter + 3 macros) is trivially small and has no business pulling in AIModule and GameplayTasks through a Developer module. Moved the whole thing into MassCore itself ??? self-contained now. If other systems ever need the same pattern we can extract to a standalone module later.

diff --git a/Engine/Plugins/Runtime/MassGameplay/Source/MassGameplayTestSuite/Private/MassEntityTemplateTest.cpp b/Engine/Plugins/Runtime/MassGameplay/Source/MassGameplayTestSuite/Private/MassEntityTemplateTest.cpp
--- a/Engine/Plugins/Runtime/MassGameplay/Source/MassGameplayTestSuite/Private/MassEntityTemplateTest.cpp
+++ b/Engine/Plugins/Runtime/MassGameplay/Source/MassGameplayTestSuite/Private/MassEntityTemplateTest.cpp
@@ -217,7 +217,7 @@ struct FEntityTest_BatchSetMissingFragment : FEntityTemplateBase
 			, {}, FMassFragmentBitSet(FTestFragment_Float::StaticStruct()));
 
 		{
-			AITEST_SCOPED_CHECK("Attempting to set value of fragment of type", 1);
+			MASS_SCOPED_ENSURE_TEST("Attempting to set value of fragment of type", 1);
 			TConstArrayView<FInstancedStruct> FragmentInstances = EntityTemplate->GetInitialFragmentValues();
 			EntityManager->BatchSetEntityFragmentValues(CreationContext->GetEntityCollections(*EntityManager.Get()), FragmentInstances);
 		}
diff --git a/Engine/Source/Developer/AITestSuite/Public/TestableEnsures.h b/Engine/Source/Developer/AITestSuite/Public/TestableEnsures.h
--- a/Engine/Source/Developer/AITestSuite/Public/TestableEnsures.h
+++ b/Engine/Source/Developer/AITestSuite/Public/TestableEnsures.h
@@ -29,7 +29,7 @@ namespace UE::AITestSuite
 	}; \
 	++UE::AITestSuite::TestsInProgress
 
-#define testableEnsureMsgf(InExpression, InFormat, ... ) \
+#define aiTestableEnsureMsgf(InExpression, InFormat, ... ) \
 	(LIKELY(!!(InExpression)) || ([&]() \
 		{ \
 			if (UNLIKELY(UE::AITestSuite::TestsInProgress > 0)) \
@@ -43,7 +43,7 @@ namespace UE::AITestSuite
 		return false; \
 		} ()))
 
-#define testableCheckf(InExpression, InFormat, ... ) \
+#define aiTestableCheckf(InExpression, InFormat, ... ) \
 	if (UNLIKELY(UE::AITestSuite::TestsInProgress > 0)) \
 	{ \
 		if (!(InExpression))\
@@ -57,7 +57,7 @@ namespace UE::AITestSuite
 		checkf(InExpression, InFormat, ##__VA_ARGS__); \
 	}
 
-#define testableCheckfReturn(InExpression, ReturnExpression, InFormat, ... ) \
+#define aiTestableCheckfReturn(InExpression, ReturnExpression, InFormat, ... ) \
 	if (UNLIKELY(UE::AITestSuite::TestsInProgress > 0)) \
 	{ \
 		if (!(InExpression)) \
diff --git a/Engine/Source/Developer/MassEntityTestSuite/Private/MassChunkCollectionTest.cpp b/Engine/Source/Developer/MassEntityTestSuite/Private/MassChunkCollectionTest.cpp
--- a/Engine/Source/Developer/MassEntityTestSuite/Private/MassChunkCollectionTest.cpp
+++ b/Engine/Source/Developer/MassEntityTestSuite/Private/MassChunkCollectionTest.cpp
@@ -245,7 +245,7 @@ struct FArchetypeEntityCollection_CreateWithInvalidDuplicatesWithPayload : FArch
 		FStructArrayView PaloadView(Payload);
 		TArray<FMassArchetypeEntityCollectionWithPayload> Result;
 		{
-			AITEST_SCOPED_CHECK("Invalid entity handle passed in", 2);
+			MASS_SCOPED_ENSURE_TEST("Invalid entity handle passed in", 2);
 			FMassArchetypeEntityCollectionWithPayload::CreateEntityRangesWithPayload(*EntityManager, EntitiesSubSet, FMassArchetypeEntityCollection::FoldDuplicates
 				, FMassGenericPayloadView(MakeArrayView(&PaloadView, 1)), Result);
 		}
diff --git a/Engine/Source/Developer/MassEntityTestSuite/Private/MassDependencyTest.cpp b/Engine/Source/Developer/MassEntityTestSuite/Private/MassDependencyTest.cpp
--- a/Engine/Source/Developer/MassEntityTestSuite/Private/MassDependencyTest.cpp
+++ b/Engine/Source/Developer/MassEntityTestSuite/Private/MassDependencyTest.cpp
@@ -417,7 +417,7 @@ struct FBadInput_Null : FDependencySolverBase
 		Processors.Reset();
 		Processors.Add(nullptr);
 		{
-			AITEST_SCOPED_CHECK(TEXT("nullptr found in Processors"), 1);
+			MASS_SCOPED_ENSURE_TEST(TEXT("nullptr found in Processors"), 1);
 			Solve();
 		}
 		AITEST_TRUE("Single nullptr input should be handled gracefully", Result.Num() == 0);
@@ -435,7 +435,7 @@ struct FBadInput_MultipleNulls : FDependencySolverBase
 		Processors.Add(nullptr);
 		Processors.Add(nullptr);
 		{
-			AITEST_SCOPED_CHECK(TEXT("nullptr found in Processors"), 3);
+			MASS_SCOPED_ENSURE_TEST(TEXT("nullptr found in Processors"), 3);
 			Solve();
 		}
 		AITEST_TRUE("Multiple nullptr inputs should be handled gracefully", Result.Num() == 0);
@@ -454,7 +454,7 @@ struct FBadInput_MixedNulls : FDependencySolverBase
 		Processors.Add(nullptr);
 		Processors.Add(NewTestProcessor<UMassTestProcessor_B>(EntityManager));
 		{
-			AITEST_SCOPED_CHECK(TEXT("nullptr found in Processors"), 2);
+			MASS_SCOPED_ENSURE_TEST(TEXT("nullptr found in Processors"), 2);
 			Solve();
 		}
 		AITEST_TRUE("Mixed nullptr and proper inputs should be handled gracefully", Result.Num() == 2);
@@ -472,7 +472,7 @@ struct FBadInput_Duplicates : FDependencySolverBase
 		Processors.Add(NewTestProcessor<UMassTestProcessor_A>(EntityManager));
 		Processors.Add(NewTestProcessor<UMassTestProcessor_A>(EntityManager));
 		{
-			AITEST_SCOPED_CHECK(TEXT("already registered. Duplicates are not supported"), 2);
+			MASS_SCOPED_ENSURE_TEST(TEXT("already registered. Duplicates are not supported"), 2);
 			Solve();
 		}
 		AITEST_TRUE("Duplicates in input should be handled gracefully", Result.Num() == 1);
@@ -548,10 +548,10 @@ struct FCircularDependency : FDependencySolverBase
 	virtual bool InstantTest() override
 	{
 		{
-			AITEST_SCOPED_CHECK(TEXT("Detected processing dependency cycle"), 1);
-			AITEST_SCOPED_CHECK(TEXT("Encountered processing dependency cycle"), 1);
+			MASS_SCOPED_ENSURE_TEST(TEXT("Detected processing dependency cycle"), 1);
+			MASS_SCOPED_ENSURE_TEST(TEXT("Encountered processing dependency cycle"), 1);
 			// this one's added since we know we log these, and don't want to trip the detectors
-			AITEST_SCOPED_CHECK(TEXT("group: None"), 4);
+			MASS_SCOPED_ENSURE_TEST(TEXT("group: None"), 4);
 			Solve();
 		}
 		// every subsequent processor is expected to depend only on the previous one since all the processors use exactly the same resources
diff --git a/Engine/Source/Developer/MassEntityTestSuite/Private/MassEntityBuilderTest.cpp b/Engine/Source/Developer/MassEntityTestSuite/Private/MassEntityBuilderTest.cpp
--- a/Engine/Source/Developer/MassEntityTestSuite/Private/MassEntityBuilderTest.cpp
+++ b/Engine/Source/Developer/MassEntityTestSuite/Private/MassEntityBuilderTest.cpp
@@ -357,7 +357,7 @@ namespace UE::Mass::Test::EntityBuilder
 				AITEST_TRUE("The secondary entity is active", EntityManager->IsEntityActive(EntityB));
 
 				{
-					AITEST_SCOPED_CHECK("Trying to commit an already committed", 1);
+					MASS_SCOPED_ENSURE_TEST("Trying to commit an already committed", 1);
 					BuilderA.Commit();
 				}
 				FMassEntityHandle EntityA2 = BuilderA.GetEntityHandle();
@@ -488,7 +488,7 @@ namespace UE::Mass::Test::EntityBuilder
 			AITEST_FALSE("(NOT) the entity has been created", EntityManager->IsEntityBuilt(ReservedHandle));
 
 			{
-				AITEST_SCOPED_CHECK("Trying to commit an already committed", 1);
+				MASS_SCOPED_ENSURE_TEST("Trying to commit an already committed", 1);
 				AITEST_INFO("Second execution of the processor shouldn't change a thing.");
 				Executor::RunProcessorsView(MakeArrayView(reinterpret_cast<UMassProcessor**>(&Processor), 1), ProcessingContext);
 			}
diff --git a/Engine/Source/Developer/MassEntityTestSuite/Private/MassEntityCollectionTest.cpp b/Engine/Source/Developer/MassEntityTestSuite/Private/MassEntityCollectionTest.cpp
--- a/Engine/Source/Developer/MassEntityTestSuite/Private/MassEntityCollectionTest.cpp
+++ b/Engine/Source/Developer/MassEntityTestSuite/Private/MassEntityCollectionTest.cpp
@@ -248,7 +248,7 @@ struct FCreateWithOutdatedCollections : FEntityCollectionTestBase
 
 		AITEST_FALSE("(NOT) After entities are moved to another archetype the InitialCollection is up to date", InitialCollection.IsUpToDate());
 		{
-			AITEST_SCOPED_CHECK("The entity collection is out of date", 1);
+			MASS_SCOPED_ENSURE_TEST("The entity collection is out of date", 1);
 			FEntityCollection EntityCollection(InitialCollection);
 			AITEST_TRUE("The EntityCollection is still empty", EntityCollection.IsEmpty());
 		}
diff --git a/Engine/Source/Developer/MassEntityTestSuite/Private/MassEntityIteratorTest.cpp b/Engine/Source/Developer/MassEntityTestSuite/Private/MassEntityIteratorTest.cpp
--- a/Engine/Source/Developer/MassEntityTestSuite/Private/MassEntityIteratorTest.cpp
+++ b/Engine/Source/Developer/MassEntityTestSuite/Private/MassEntityIteratorTest.cpp
@@ -119,7 +119,7 @@ struct FIterator_QuerylessIterator : FEntityTestBase
 	{
 		FMassExecutionContext LocalContext(*EntityManager.Get());
 		{
-			AITEST_SCOPED_CHECK("no entity query is being executed", 1);
+			MASS_SCOPED_ENSURE_TEST("no entity query is being executed", 1);
 			FMassExecutionContext::FEntityIterator FailedIterator = LocalContext.CreateEntityIterator();
 
 			AITEST_FALSE("(Not) Created iterator is valid", bool(FailedIterator));
@@ -143,7 +143,7 @@ struct FIterator_ProcessorlessIterator : FEntityTestBase
 	{
 		FMassExecutionContext LocalContext(*EntityManager.Get());
 		{
-			AITEST_SCOPED_CHECK("no entity query is being executed", 1);
+			MASS_SCOPED_ENSURE_TEST("no entity query is being executed", 1);
 			FMassExecutionContext::FEntityIterator FailedIterator = LocalContext.CreateEntityIterator();
 
 			AITEST_FALSE("(Not) Created iterator is valid", bool(FailedIterator));
diff --git a/Engine/Source/Developer/MassEntityTestSuite/Private/MassExecutionTest.cpp b/Engine/Source/Developer/MassEntityTestSuite/Private/MassExecutionTest.cpp
--- a/Engine/Source/Developer/MassEntityTestSuite/Private/MassExecutionTest.cpp
+++ b/Engine/Source/Developer/MassEntityTestSuite/Private/MassExecutionTest.cpp
@@ -69,7 +69,7 @@ struct FExecution_SingleNullProcessor : FExecutionTestBase
 		TArray<UMassProcessor*> Processors;
 		Processors.Add(nullptr);
 
-		AITEST_SCOPED_CHECK(TEXT("Processors contains nullptr"), 1);
+		MASS_SCOPED_ENSURE_TEST(TEXT("Processors contains nullptr"), 1);
 		// note that using RunProcessorsView is to bypass reasonable tests UE::Mass::Executor::Run(Pipeline,...) does that are 
 		// reported via ensures which are not handled by the automation framework
 		UE::Mass::Executor::RunProcessorsView(Processors, ProcessingContext);
@@ -111,7 +111,7 @@ struct FExecution_MultipleNullProcessors : FExecutionTestBase
 		Processors.Add(nullptr);
 		Processors.Add(nullptr);
 
-		AITEST_SCOPED_CHECK(TEXT("Processors contains nullptr"), 1);
+		MASS_SCOPED_ENSURE_TEST(TEXT("Processors contains nullptr"), 1);
 		// note that using RunProcessorsView is to bypass reasonable tests UE::Mass::Executor::Run(Pipeline,...) does that are 
 		// reported via ensures which are not handled by the automation framework
 		UE::Mass::Executor::RunProcessorsView(Processors, ProcessingContext);
diff --git a/Engine/Source/Developer/MassEntityTestSuite/Private/MassObserverTest.cpp b/Engine/Source/Developer/MassEntityTestSuite/Private/MassObserverTest.cpp
--- a/Engine/Source/Developer/MassEntityTestSuite/Private/MassObserverTest.cpp
+++ b/Engine/Source/Developer/MassEntityTestSuite/Private/MassObserverTest.cpp
@@ -644,7 +644,7 @@ struct FObserverChangingComposition_Sync : FFragmentTestBase
 		ObserverManager.AddObserverInstance(FFragmentStruct::StaticStruct(), OperationFlagsObserved, ObserverProcessor);
 
 		{
-			AITEST_SCOPED_CHECK("Use asynchronous API instead", 1);
+			MASS_SCOPED_ENSURE_TEST("Use asynchronous API instead", 1);
 
 			ObserverProcessor->ForEachEntityChunkExecutionFunction = [EntityManager = EntityManager](FMassExecutionContext& Context)
 			{
@@ -656,7 +656,7 @@ struct FObserverChangingComposition_Sync : FFragmentTestBase
 			AITEST_EQUAL("Number of entities in the original archetype, no moves expected", EntityManager->DebugGetArchetypeEntitiesCount(OriginalArchetype), EntitiesToSpawn);
 		}
 		{
-			AITEST_SCOPED_CHECK("Use asynchronous API instead", 1);
+			MASS_SCOPED_ENSURE_TEST("Use asynchronous API instead", 1);
 
 			ObserverProcessor->ForEachEntityChunkExecutionFunction = [EntityManager = EntityManager, OriginalArchetype](FMassExecutionContext& Context)
 			{
diff --git a/Engine/Source/Developer/MassEntityTestSuite/Private/MassRelationTest.cpp b/Engine/Source/Developer/MassEntityTestSuite/Private/MassRelationTest.cpp
--- a/Engine/Source/Developer/MassEntityTestSuite/Private/MassRelationTest.cpp
+++ b/Engine/Source/Developer/MassEntityTestSuite/Private/MassRelationTest.cpp
@@ -513,7 +513,7 @@ namespace UE::Mass::Test::Relations
 
 			const FMassEntityHandle Entity = EntityManager->CreateEntity(OriginalArchetype);
 			{
-				AITEST_SCOPED_CHECK("subject and object to be different", 1);
+				MASS_SCOPED_ENSURE_TEST("subject and object to be different", 1);
 				RelationshipManagerPtr->CreateRelationInstance(RelationTypeHandle, Entity, Entity);
 			}
 			const FMassArchetypeHandle FinalArchetype = EntityManager->GetArchetypeForEntity(Entity);
@@ -540,7 +540,7 @@ namespace UE::Mass::Test::Relations
 			RelationshipManagerPtr->CreateRelationInstance(RelationTypeHandle, Entities[2], Entities[1]);
 			{
 				// this relation registration should fail as it would create a cycle
-				AITEST_SCOPED_CHECK("result in a cycle", 1);
+				MASS_SCOPED_ENSURE_TEST("result in a cycle", 1);
 				RelationshipManagerPtr->CreateRelationInstance(RelationTypeHandle, Entities[0], Entities[2]);
 			}
 
@@ -645,7 +645,7 @@ namespace UE::Mass::Test::Relations
 
 			// try creating relations with mismatching pairs
 			{
-				AITEST_SCOPED_CHECK("Relation Objects count need match Subjects", 1);
+				MASS_SCOPED_ENSURE_TEST("Relation Objects count need match Subjects", 1);
 				TArray<FMassEntityHandle> RelationInstances = RelationshipManagerPtr->CreateRelationInstances(LocalRelationTypeHandle
 					, TArrayView<FMassEntityHandle>(&Entities[0], 1), TArrayView<FMassEntityHandle>(&Entities[1], 2));
 				AITEST_EQUAL("A: Mismatching relation instances created count", RelationInstances.Num(), 0);
@@ -654,15 +654,15 @@ namespace UE::Mass::Test::Relations
 				TArray<FMassEntityHandle> Subjects;
 				TArray<FMassEntityHandle> Objects;
 
-				AITEST_SCOPED_CHECK("Relation needs a valid Subject entity", 1);
+				MASS_SCOPED_ENSURE_TEST("Relation needs a valid Subject entity", 1);
 				Subjects.Add(FMassEntityHandle());
 				Objects.Add(Entities[0]);
 
-				AITEST_SCOPED_CHECK("Relation needs a valid Object entity", 1);
+				MASS_SCOPED_ENSURE_TEST("Relation needs a valid Object entity", 1);
 				Subjects.Add(Entities[0]);
 				Objects.Add(FMassEntityHandle());
 				
-				AITEST_SCOPED_CHECK("Hierarchical relation requires the subject and object to be different", 1);
+				MASS_SCOPED_ENSURE_TEST("Hierarchical relation requires the subject and object to be different", 1);
 				Subjects.Add(Entities[0]);
 				Objects.Add(Entities[0]);
 
@@ -670,7 +670,7 @@ namespace UE::Mass::Test::Relations
 				Subjects.Add(Entities[1]);
 				Objects.Add(Entities[2]);
 
-				AITEST_SCOPED_CHECK("Relation between the two entities already exists", 1);
+				MASS_SCOPED_ENSURE_TEST("Relation between the two entities already exists", 1);
 				Subjects.Add(Entities[0]);
 				Objects.Add(Entities[1]);
 
@@ -697,20 +697,20 @@ namespace UE::Mass::Test::Relations
 				EntityManager->GetTypeManager().RegisterType(MoveTemp(LocalRelationTraits));
 			}
 			{
-				AITEST_SCOPED_CHECK("Modifying relationship after registration done is not supported", 1);
+				MASS_SCOPED_ENSURE_TEST("Modifying relationship after registration done is not supported", 1);
 				FRelationTypeTraits LocalRelationTraits = CreateRelationTypeDescription("TestFragmentA");
 				// make it different
 				LocalRelationTraits.bHierarchical = !LocalRelationTraits.bHierarchical;
 				EntityManager->GetTypeManager().RegisterType(MoveTemp(LocalRelationTraits));
 			}
 			{
-				AITEST_SCOPED_CHECK("Modifying relationship after registration done is not supported", 1);
+				MASS_SCOPED_ENSURE_TEST("Modifying relationship after registration done is not supported", 1);
 				FRelationTypeTraits LocalRelationTraits = CreateRelationTypeDescription("TestFragmentB");
 				// register with existing name
 				EntityManager->GetTypeManager().RegisterType(MoveTemp(LocalRelationTraits));
 			}
 			{
-				AITEST_SCOPED_CHECK("Modifying relationship after registration done is not supported", 1);
+				MASS_SCOPED_ENSURE_TEST("Modifying relationship after registration done is not supported", 1);
 				FRelationTypeTraits LocalRelationTraits = CreateRelationTypeDescription("TestFragmentA");
 				// register existing traits - source and object elements are supposed to be unique
 				EntityManager->GetTypeManager().RegisterType(MoveTemp(LocalRelationTraits));
diff --git a/Engine/Source/Developer/MassEntityTestSuite/Private/MassSharedFragmentsTest.cpp b/Engine/Source/Developer/MassEntityTestSuite/Private/MassSharedFragmentsTest.cpp
--- a/Engine/Source/Developer/MassEntityTestSuite/Private/MassSharedFragmentsTest.cpp
+++ b/Engine/Source/Developer/MassEntityTestSuite/Private/MassSharedFragmentsTest.cpp
@@ -262,7 +262,7 @@ struct FSharedFragmentValues_Hash : FExecutionTestBase
 		ValuestNonConstIntFloat.Add(FSharedStruct::Make<FTestSharedFragment_Int>(TestIntValue));
 		ValuestNonConstIntFloat.Add(FSharedStruct::Make<FTestSharedFragment_Float>(TestFloatValue));
 
-		AITEST_SCOPED_CHECK("Expecting the containers to be sorted", 1);
+		MASS_SCOPED_ENSURE_TEST("Expecting the containers to be sorted", 1);
 		const uint32 EmptyHash = ValuestNonConstIntFloat.CalculateHash();
 		AITEST_EQUAL("Expecting unsorted collection hashing to result in 0", EmptyHash, 0u);
 
@@ -529,7 +529,7 @@ struct FConstSharedFragment_AddToEntity : FEntityTestBase
 		AITEST_NULL("Initially the entity is not expected to have the shared fragment", EntitySharedFragment);
 
 		{
-			AITEST_SCOPED_CHECK("shared struct requires to be registered first", 1);
+			MASS_SCOPED_ENSURE_TEST("shared struct requires to be registered first", 1);
 			const FConstSharedStruct UnregisteredSharedFragmentInstance = FSharedStruct::Make(FragmentInstance);
 			const bool bSuccessfullyAddedConstSharedFragment =  EntityManager->AddConstSharedFragmentToEntity(EntityHandle, UnregisteredSharedFragmentInstance);
 			AITEST_FALSE("Adding unregistered const shared fragment struct should fail", bSuccessfullyAddedConstSharedFragment);
@@ -552,7 +552,7 @@ struct FConstSharedFragment_AddToEntity : FEntityTestBase
 		const FTestConstSharedFragment_Int DifferentFragmentInstance(DifferentTestIntValue);
 		const FConstSharedStruct DifferentConstSharedFragmentInstance = EntityManager->GetOrCreateConstSharedFragment(DifferentFragmentInstance);
 
-		AITEST_SCOPED_CHECK("Changing shared fragment value of entities is not supported", 1);
+		MASS_SCOPED_ENSURE_TEST("Changing shared fragment value of entities is not supported", 1);
 
 		const bool bSuccessfullyAddedConstSharedFragment = EntityManager->AddConstSharedFragmentToEntity(EntityHandle, DifferentConstSharedFragmentInstance);
 		AITEST_FALSE("Adding existing const shared fragment type should fail", bSuccessfullyAddedConstSharedFragment);
@@ -801,8 +801,8 @@ struct FSharedFragment_BatchSetAttempt : FEntityTestBase
 		
 		// attempting to add the same values again should fail with checks and ensures
 		{	
-			AITEST_SCOPED_CHECK("Setting shared fragment values without archetype change is not supported", 1);
-			AITEST_SCOPED_CHECK("Trying to set shared fragment values, without adding new shared fragments", 1);
+			MASS_SCOPED_ENSURE_TEST("Setting shared fragment values without archetype change is not supported", 1);
+			MASS_SCOPED_ENSURE_TEST("Trying to set shared fragment values, without adding new shared fragments", 1);
 			EntityManager->BatchAddSharedFragmentsForEntities(MakeArrayView(&EntityCollection, 1), SharedIntValues);
 		}
 		{
@@ -810,8 +810,8 @@ struct FSharedFragment_BatchSetAttempt : FEntityTestBase
 			FConstSharedStruct OtherConstSharedFragment = FConstSharedStruct::Make<FTestConstSharedFragment_Int>(OtherTestIntValue);
 			DifferentSharedIntValues.Add(OtherConstSharedFragment);
 
-			AITEST_SCOPED_CHECK("Setting shared fragment values without archetype change is not supported", 1);
-			AITEST_SCOPED_CHECK("Trying to set shared fragment values, without adding new shared fragments", 1);
+			MASS_SCOPED_ENSURE_TEST("Setting shared fragment values without archetype change is not supported", 1);
+			MASS_SCOPED_ENSURE_TEST("Trying to set shared fragment values, without adding new shared fragments", 1);
 			EntityManager->BatchAddSharedFragmentsForEntities(MakeArrayView(&EntityCollection, 1), DifferentSharedIntValues);
 		}
 
@@ -836,7 +836,7 @@ struct FSharedFragment_BatchAddToEmpty : FEntityTestBase
 		FMassArchetypeEntityCollection EntityCollection(FMassArchetypeHandle(), ReservedEntityHandles, FMassArchetypeEntityCollection::EDuplicatesHandling::NoDuplicates);
 		// attempting to add the values before the entities are created is not a valid operation
 		{
-			AITEST_SCOPED_CHECK("Adding shared fragments to archetype-less entities is not supported", 1);
+			MASS_SCOPED_ENSURE_TEST("Adding shared fragments to archetype-less entities is not supported", 1);
 			EntityManager->BatchAddSharedFragmentsForEntities(MakeArrayView(&EntityCollection, 1), SharedIntValues);
 		}
 		return true;
diff --git a/Engine/Source/Developer/MassEntityTestSuite/Public/MassEntityTestTypes.h b/Engine/Source/Developer/MassEntityTestSuite/Public/MassEntityTestTypes.h
--- a/Engine/Source/Developer/MassEntityTestSuite/Public/MassEntityTestTypes.h
+++ b/Engine/Source/Developer/MassEntityTestSuite/Public/MassEntityTestTypes.h
@@ -21,6 +21,20 @@
 
 #define UE_API MASSENTITYTESTSUITE_API
 
+/** Works with checks too, but didn't want to make the name too lon */
+#define MASS_SCOPED_ENSURE_TEST(PartialTextToExpect, NumOccurrences) \
+	GetTestRunner().AddExpectedErrorPlain(PartialTextToExpect, EAutomationExpectedErrorFlags::Contains, NumOccurrences); \
+	ON_SCOPE_EXIT \
+	{ \
+		--UE::Mass::Test::TestsInProgress; \
+		FAutomationTestExecutionInfo ExecutionInfo; \
+		GetTestRunner().GetExecutionInfo(ExecutionInfo); \
+		if (ExecutionInfo.GetWarningTotal() + ExecutionInfo.GetWarningTotal() > 0) \
+		{ \
+			UE::AITestSuite::ConditionallyBreakOnTestFail(); \
+		} \
+	}; \
+	++UE::Mass::Test::TestsInProgress
 
 struct FMassEntityManager;
 struct FMassProcessingPhaseManager;
diff --git a/Engine/Source/Programs/MassTests/Private/MassChunkCollectionTest.cpp b/Engine/Source/Programs/MassTests/Private/MassChunkCollectionTest.cpp
--- a/Engine/Source/Programs/MassTests/Private/MassChunkCollectionTest.cpp
+++ b/Engine/Source/Programs/MassTests/Private/MassChunkCollectionTest.cpp
@@ -205,7 +205,7 @@ TEST_CASE_METHOD(FChunkCollectionFixture, "Mass::ArchetypeEntityCollection::Crea
 	}
 }
 
-// @todo Re-enable once test body is implemented. Original used AITEST_SCOPED_CHECK("Invalid entity handle passed in", 2).
+// @todo Re-enable once test body is implemented. Original used MASS_SCOPED_ENSURE_TEST("Invalid entity handle passed in", 2).
 DISABLED_TEST_CASE_METHOD(FChunkCollectionFixture, "Mass::ArchetypeEntityCollection::Create::InvalidDuplicatesWithPayload", "[Mass][ArchetypeEntityCollection]")
 {
 }
diff --git a/Engine/Source/Programs/MassTests/Private/MassSharedFragmentsTest.cpp b/Engine/Source/Programs/MassTests/Private/MassSharedFragmentsTest.cpp
--- a/Engine/Source/Programs/MassTests/Private/MassSharedFragmentsTest.cpp
+++ b/Engine/Source/Programs/MassTests/Private/MassSharedFragmentsTest.cpp
@@ -303,7 +303,7 @@ TEST_CASE_METHOD(FMassLLTFixture, "Mass::SharedFragments::Hash", "[Mass][SharedF
 	ValuestNonConstIntFloat.Add(FSharedStruct::Make<FTestSharedFragment_Float>(TestFloatValue));
 
 	// Note: Calling CalculateHash() on unsorted collection intentionally triggers ensure(bSorted),
-	// which is fatal in LLT. The original test verified it returns 0 via AITEST_SCOPED_CHECK.
+	// which is fatal in LLT. The original test verified it returns 0 via MASS_SCOPED_ENSURE_TEST.
 	// Skipping the unsorted hash check here.
 
 	ValuestNonConstIntFloat.Sort();
diff --git a/Engine/Source/Runtime/Mass/MassCore/MassCore.Build.cs b/Engine/Source/Runtime/Mass/MassCore/MassCore.Build.cs
--- a/Engine/Source/Runtime/Mass/MassCore/MassCore.Build.cs
+++ b/Engine/Source/Runtime/Mass/MassCore/MassCore.Build.cs
@@ -17,13 +17,6 @@ namespace UnrealBuildTool.Rules
 					"TraceLog",
 				}
 			);
-
-			if (Target.Configuration != UnrealTargetConfiguration.Shipping
-				&& Target.Configuration != UnrealTargetConfiguration.Test)
-			{
-				// @todo: the dependency to be removed in near future.
-				PrivateDependencyModuleNames.Add("AITestSuite");
-			}
 		}
 	}
 }
diff --git a/Engine/Source/Runtime/Mass/MassCore/Private/MassTestableEnsures.cpp b/Engine/Source/Runtime/Mass/MassCore/Private/MassTestableEnsures.cpp
new file mode 100644
--- /dev/null
+++ b/Engine/Source/Runtime/Mass/MassCore/Private/MassTestableEnsures.cpp
@@ -0,0 +1,14 @@
+// Copyright Epic Games, Inc. All Rights Reserved.
+
+#include "MassTestableEnsures.h"
+
+#if !UE_BUILD_SHIPPING && !UE_BUILD_TEST
+
+DEFINE_LOG_CATEGORY(LogMassTestableEnsures);
+
+namespace UE::Mass::Test
+{
+	int32 TestsInProgress = 0;
+}
+
+#endif // !UE_BUILD_SHIPPING && !UE_BUILD_TEST
diff --git a/Engine/Source/Runtime/Mass/MassCore/Public/MassTestableEnsures.h b/Engine/Source/Runtime/Mass/MassCore/Public/MassTestableEnsures.h
--- a/Engine/Source/Runtime/Mass/MassCore/Public/MassTestableEnsures.h
+++ b/Engine/Source/Runtime/Mass/MassCore/Public/MassTestableEnsures.h
@@ -2,10 +2,76 @@
 
 #pragma once
 
-#ifdef WITH_AITESTSUITE
-#include "TestableEnsures.h"
-#else
+// Self-contained testable assertion macros for Mass.
+// These allow Mass runtime checks to log-and-continue during automated tests
+// instead of crashing, while behaving as standard ensure/check in production.
+//
+// @todo: if other engine systems need testable ensures, consider extracting
+// this to a standalone TestableAssertions module rather than duplicating.
+
+#include "Misc/AssertionMacros.h"
+
+#if !UE_BUILD_SHIPPING && !UE_BUILD_TEST
+
+#include "Logging/LogCategory.h"
+
+MASSCORE_API DECLARE_LOG_CATEGORY_EXTERN(LogMassTestableEnsures, Log, All);
+
+namespace UE::Mass::Test
+{
+	// Counter indicating how many test scopes are active. When > 0, testable
+	// macros log instead of crashing.
+	MASSCORE_API extern int32 TestsInProgress;
+}
+
+#define testableEnsureMsgf(InExpression, InFormat, ... ) \
+	(LIKELY(!!(InExpression)) || ([&]() \
+		{ \
+			if (UNLIKELY(UE::Mass::Test::TestsInProgress > 0)) \
+			{ \
+				UE_LOG(LogMassTestableEnsures, Warning, InFormat, ##__VA_ARGS__); \
+			} \
+			else \
+			{ \
+				ensureMsgf(InExpression, InFormat, ##__VA_ARGS__); \
+			} \
+		return false; \
+		} ()))
+
+#define testableCheckf(InExpression, InFormat, ... ) \
+	if (UNLIKELY(UE::Mass::Test::TestsInProgress > 0)) \
+	{ \
+		if (!(InExpression))\
+		{ \
+			UE_LOG(LogMassTestableEnsures, Error, InFormat, ##__VA_ARGS__); \
+			return; \
+		}\
+	} \
+	else \
+	{ \
+		checkf(InExpression, InFormat, ##__VA_ARGS__); \
+	}
+
+#define testableCheckfReturn(InExpression, ReturnExpression, InFormat, ... ) \
+	if (UNLIKELY(UE::Mass::Test::TestsInProgress > 0)) \
+	{ \
+		if (!(InExpression)) \
+		{ \
+			UE_LOG(LogMassTestableEnsures, Error, InFormat, ##__VA_ARGS__); \
+			ReturnExpression; \
+		} \
+	} \
+	else \
+	{ \
+		checkf(InExpression, InFormat, ##__VA_ARGS__); \
+	}
+
+#else // UE_BUILD_SHIPPING || UE_BUILD_TEST
+
+struct FMassTestableEnsureScope {};
+
 #define testableEnsureMsgf ensureMsgf
 #define testableCheckf checkf
 #define testableCheckfReturn(InExpression, ReturnValue, InFormat, ... ) checkf(InExpression, InFormat, ##__VA_ARGS__)
-#endif 
+
+#endif
diff --git a/Engine/Source/Runtime/MassEntity/MassEntity.Build.cs b/Engine/Source/Runtime/MassEntity/MassEntity.Build.cs
--- a/Engine/Source/Runtime/MassEntity/MassEntity.Build.cs
+++ b/Engine/Source/Runtime/MassEntity/MassEntity.Build.cs
@@ -27,13 +27,6 @@ public class MassEntity : ModuleRules
 			PrivateDependencyModuleNames.Add("EditorSubsystem");
 		}
 
-		if (Target.Configuration != UnrealTargetConfiguration.Shipping
-			&& Target.Configuration != UnrealTargetConfiguration.Test)
-		{
-			// pulling this one in for the testableEnsureMsgf
-			PrivateDependencyModuleNames.Add("AITestSuite");
-		}
-
 		if (Target.bBuildDeveloperTools)
 		{
 			DynamicallyLoadedModuleNames.Add("MassEntityTestSuite");
