[Mass]
Extend BatchRemoveElement to support shared fragments and add unit tests

- Extend IsElement() to accept shared and const shared fragment types
- Extend BatchRemoveElementFromEntities and RemoveElementFromEntity to dispatch shared and const shared fragment removal via new InternalBatchRemoveSharedFragmentsFromEntities
- Add testableCheckf to Add element APIs rejecting shared fragments with clear message
- Add unit tests for FMassCommandAddElement/FMassCommandRemoveElement covering tags, fragments, shared fragments, const shared fragments, mixed types, and negative cases

#Synced-CL 52443815

diff --git a/Engine/Source/Developer/MassEntityTestSuite/Private/MassCommandsTest.cpp b/Engine/Source/Developer/MassEntityTestSuite/Private/MassCommandsTest.cpp
--- a/Engine/Source/Developer/MassEntityTestSuite/Private/MassCommandsTest.cpp
+++ b/Engine/Source/Developer/MassEntityTestSuite/Private/MassCommandsTest.cpp
@@ -1,11 +1,12 @@
 // Copyright Epic Games, Inc. All Rights Reserved.
 
 #include "MassEntityManager.h"
-#include "MassProcessingTypes.h"
+#include "MassCommands.h"
 #include "MassEntityTestTypes.h"
 #include "MassEntityTypes.h"
 #include "MassEntityView.h"
 #include "MassExecutionContext.h"
+#include "MassProcessingTypes.h"
 #include "Algo/Sort.h"
 #include "Algo/RandomShuffle.h"
 
@@ -345,6 +346,380 @@ struct FCommands_MoveHandleArrays : FEntityTestBase
 };
 IMPLEMENT_AI_INSTANT_TEST(FCommands_MoveHandleArrays, "System.Mass.Commands.MoveHandleArrays");
 
+//----------------------------------------------------------------------//
+// FMassCommandAddElement / FMassCommandRemoveElement tests
+//----------------------------------------------------------------------//
+struct FCommands_AddElement_OnlyTags : FEntityTestBase
+{
+	virtual bool InstantTest() override
+	{
+		const int32 Count = 5;
+		TArray<FMassEntityHandle> Entities;
+		EntityManager->BatchCreateEntities(IntsArchetype, Count, Entities);
+
+		for (const FMassEntityHandle& Entity : Entities)
+		{
+			AITEST_FALSE(TEXT("Entity should not have TagA before the command"), EntityManager->DoesEntityHaveElement(Entity, FTestTag_A::StaticStruct()));
+		}
+
+		FMassCommandAddElement& Command = EntityManager->Defer().PushUniqueCommand<FMassCommandAddElement>(FTestTag_A::StaticStruct());
+		for (const FMassEntityHandle& Entity : Entities)
+		{
+			Command.Add(Entity);
+		}
+
+		EntityManager->FlushCommands();
+
+		for (const FMassEntityHandle& Entity : Entities)
+		{
+			AITEST_TRUE(TEXT("Entity should have TagA after flushing"), EntityManager->DoesEntityHaveElement(Entity, FTestTag_A::StaticStruct()));
+			AITEST_TRUE(TEXT("Entity should still have FTestFragment_Int"), EntityManager->DoesEntityHaveElement(Entity, FTestFragment_Int::StaticStruct()));
+		}
+
+		return true;
+	}
+};
+IMPLEMENT_AI_INSTANT_TEST(FCommands_AddElement_OnlyTags, "System.Mass.Commands.AddElement.OnlyTags");
+
+struct FCommands_AddElement_OnlyFragments : FEntityTestBase
+{
+	virtual bool InstantTest() override
+	{
+		const int32 Count = 5;
+		TArray<FMassEntityHandle> Entities;
+		EntityManager->BatchCreateEntities(IntsArchetype, Count, Entities);
+
+		for (const FMassEntityHandle& Entity : Entities)
+		{
+			AITEST_FALSE(TEXT("Entity should not have FTestFragment_Float before the command"), EntityManager->DoesEntityHaveElement(Entity, FTestFragment_Float::StaticStruct()));
+		}
+
+		FMassCommandAddElement& Command = EntityManager->Defer().PushUniqueCommand<FMassCommandAddElement>(FTestFragment_Float::StaticStruct());
+		for (const FMassEntityHandle& Entity : Entities)
+		{
+			Command.Add(Entity);
+		}
+
+		EntityManager->FlushCommands();
+
+		for (const FMassEntityHandle& Entity : Entities)
+		{
+			AITEST_TRUE(TEXT("Entity should have FTestFragment_Float after flushing"), EntityManager->DoesEntityHaveElement(Entity, FTestFragment_Float::StaticStruct()));
+			AITEST_TRUE(TEXT("Entity should still have FTestFragment_Int"), EntityManager->DoesEntityHaveElement(Entity, FTestFragment_Int::StaticStruct()));
+		}
+
+		return true;
+	}
+};
+IMPLEMENT_AI_INSTANT_TEST(FCommands_AddElement_OnlyFragments, "System.Mass.Commands.AddElement.OnlyFragments");
+
+struct FCommands_RemoveElement_OnlyTags : FEntityTestBase
+{
+	virtual bool InstantTest() override
+	{
+		const int32 Count = 5;
+		const FMassArchetypeHandle IntsTagAArchetype = EntityManager->CreateArchetype({ FTestFragment_Int::StaticStruct(), FTestTag_A::StaticStruct() });
+
+		TArray<FMassEntityHandle> Entities;
+		EntityManager->BatchCreateEntities(IntsTagAArchetype, Count, Entities);
+
+		for (const FMassEntityHandle& Entity : Entities)
+		{
+			AITEST_TRUE(TEXT("Entity should have TagA before removal"), EntityManager->DoesEntityHaveElement(Entity, FTestTag_A::StaticStruct()));
+		}
+
+		FMassCommandRemoveElement& Command = EntityManager->Defer().PushUniqueCommand<FMassCommandRemoveElement>(FTestTag_A::StaticStruct());
+		for (const FMassEntityHandle& Entity : Entities)
+		{
+			Command.Add(Entity);
+		}
+
+		EntityManager->FlushCommands();
+
+		for (const FMassEntityHandle& Entity : Entities)
+		{
+			AITEST_FALSE(TEXT("Entity should no longer have TagA after flushing"), EntityManager->DoesEntityHaveElement(Entity, FTestTag_A::StaticStruct()));
+			AITEST_TRUE(TEXT("Entity should still have FTestFragment_Int"), EntityManager->DoesEntityHaveElement(Entity, FTestFragment_Int::StaticStruct()));
+		}
+
+		return true;
+	}
+};
+IMPLEMENT_AI_INSTANT_TEST(FCommands_RemoveElement_OnlyTags, "System.Mass.Commands.RemoveElement.OnlyTags");
+
+struct FCommands_RemoveElement_OnlyFragments : FEntityTestBase
+{
+	virtual bool InstantTest() override
+	{
+		const int32 Count = 5;
+		TArray<FMassEntityHandle> Entities;
+		EntityManager->BatchCreateEntities(FloatsIntsArchetype, Count, Entities);
+
+		for (const FMassEntityHandle& Entity : Entities)
+		{
+			AITEST_TRUE(TEXT("Entity should have FTestFragment_Float before removal"), EntityManager->DoesEntityHaveElement(Entity, FTestFragment_Float::StaticStruct()));
+		}
+
+		FMassCommandRemoveElement& Command = EntityManager->Defer().PushUniqueCommand<FMassCommandRemoveElement>(FTestFragment_Float::StaticStruct());
+		for (const FMassEntityHandle& Entity : Entities)
+		{
+			Command.Add(Entity);
+		}
+
+		EntityManager->FlushCommands();
+
+		for (const FMassEntityHandle& Entity : Entities)
+		{
+			AITEST_FALSE(TEXT("Entity should no longer have FTestFragment_Float"), EntityManager->DoesEntityHaveElement(Entity, FTestFragment_Float::StaticStruct()));
+			AITEST_TRUE(TEXT("Entity should still have FTestFragment_Int"), EntityManager->DoesEntityHaveElement(Entity, FTestFragment_Int::StaticStruct()));
+			AITEST_EQUAL(TEXT("Entity should have moved to IntsArchetype"), EntityManager->GetArchetypeForEntity(Entity), IntsArchetype);
+		}
+
+		return true;
+	}
+};
+IMPLEMENT_AI_INSTANT_TEST(FCommands_RemoveElement_OnlyFragments, "System.Mass.Commands.RemoveElement.OnlyFragments");
+
+struct FCommands_RemoveElement_OnlyConstSharedFragments : FEntityTestBase
+{
+	virtual bool InstantTest() override
+	{
+		const int32 Count = 5;
+		TArray<FMassEntityHandle> Entities;
+		EntityManager->BatchCreateEntities(IntsArchetype, Count, Entities);
+
+		// Add a const shared fragment to all entities
+		FTestConstSharedFragment_Int SharedFragmentInstance(42);
+		FConstSharedStruct ConstSharedStruct = EntityManager->GetOrCreateConstSharedFragment(SharedFragmentInstance);
+		for (const FMassEntityHandle& Entity : Entities)
+		{
+			EntityManager->AddConstSharedFragmentToEntity(Entity, ConstSharedStruct);
+		}
+
+		for (const FMassEntityHandle& Entity : Entities)
+		{
+			AITEST_TRUE(TEXT("Entity should have const shared fragment before removal"), EntityManager->DoesEntityHaveElement(Entity, FTestConstSharedFragment_Int::StaticStruct()));
+		}
+
+		FMassCommandRemoveElement& Command = EntityManager->Defer().PushUniqueCommand<FMassCommandRemoveElement>(FTestConstSharedFragment_Int::StaticStruct());
+		for (const FMassEntityHandle& Entity : Entities)
+		{
+			Command.Add(Entity);
+		}
+
+		EntityManager->FlushCommands();
+
+		for (const FMassEntityHandle& Entity : Entities)
+		{
+			AITEST_FALSE(TEXT("Entity should no longer have const shared fragment"), EntityManager->DoesEntityHaveElement(Entity, FTestConstSharedFragment_Int::StaticStruct()));
+			AITEST_TRUE(TEXT("Entity should still have FTestFragment_Int"), EntityManager->DoesEntityHaveElement(Entity, FTestFragment_Int::StaticStruct()));
+		}
+
+		return true;
+	}
+};
+IMPLEMENT_AI_INSTANT_TEST(FCommands_RemoveElement_OnlyConstSharedFragments, "System.Mass.Commands.RemoveElement.OnlyConstSharedFragments");
+
+struct FCommands_RemoveElement_MixedTypes : FEntityTestBase
+{
+	virtual bool InstantTest() override
+	{
+		constexpr int32 Count = 5;
+		const FMassArchetypeHandle FloatsIntsTagAArchetype = EntityManager->CreateArchetype(
+			{ FTestFragment_Int::StaticStruct(), FTestFragment_Float::StaticStruct(), FTestTag_A::StaticStruct() });
+
+		TArray<FMassEntityHandle> Entities;
+		EntityManager->BatchCreateEntities(FloatsIntsTagAArchetype, Count, Entities);
+
+		// Add a const shared fragment to all entities
+		FTestConstSharedFragment_Int ConstSharedInstance(42);
+		FConstSharedStruct ConstSharedStruct = EntityManager->GetOrCreateConstSharedFragment(ConstSharedInstance);
+		for (const FMassEntityHandle& Entity : Entities)
+		{
+			EntityManager->AddConstSharedFragmentToEntity(Entity, ConstSharedStruct);
+		}
+
+		// Add a mutable shared fragment to all entities
+		FTestSharedFragment_Int SharedInstance(99);
+		FSharedStruct SharedStruct = FSharedStruct::Make(SharedInstance);
+		FMassArchetypeSharedFragmentValues SharedFragmentValues;
+		SharedFragmentValues.Add(SharedStruct);
+
+		TArray<FMassArchetypeEntityCollection> EntityCollections;
+		UE::Mass::Utils::CreateEntityCollections(*EntityManager.Get(), Entities, FMassArchetypeEntityCollection::NoDuplicates, EntityCollections);
+		EntityManager->BatchAddSharedFragmentsForEntities(EntityCollections, SharedFragmentValues);
+
+		for (const FMassEntityHandle& Entity : Entities)
+		{
+			AITEST_TRUE(TEXT("Entity should have all elements before removal"),
+				EntityManager->DoesEntityHaveElement(Entity, FTestFragment_Float::StaticStruct())
+				&& EntityManager->DoesEntityHaveElement(Entity, FTestTag_A::StaticStruct())
+				&& EntityManager->DoesEntityHaveElement(Entity, FTestConstSharedFragment_Int::StaticStruct())
+				&& EntityManager->DoesEntityHaveElement(Entity, FTestSharedFragment_Int::StaticStruct()));
+		}
+
+		// Push four separate RemoveElement commands for all element types
+		FMassCommandRemoveElement& RemoveFragmentCmd = EntityManager->Defer().PushUniqueCommand<FMassCommandRemoveElement>(FTestFragment_Float::StaticStruct());
+		FMassCommandRemoveElement& RemoveTagCmd = EntityManager->Defer().PushUniqueCommand<FMassCommandRemoveElement>(FTestTag_A::StaticStruct());
+		FMassCommandRemoveElement& RemoveConstSharedCmd = EntityManager->Defer().PushUniqueCommand<FMassCommandRemoveElement>(FTestConstSharedFragment_Int::StaticStruct());
+		FMassCommandRemoveElement& RemoveSharedCmd = EntityManager->Defer().PushUniqueCommand<FMassCommandRemoveElement>(FTestSharedFragment_Int::StaticStruct());
+		for (const FMassEntityHandle& Entity : Entities)
+		{
+			RemoveFragmentCmd.Add(Entity);
+			RemoveTagCmd.Add(Entity);
+			RemoveConstSharedCmd.Add(Entity);
+			RemoveSharedCmd.Add(Entity);
+		}
+
+		EntityManager->FlushCommands();
+
+		for (const FMassEntityHandle& Entity : Entities)
+		{
+			AITEST_TRUE(TEXT("Entity should still have FTestFragment_Int"), EntityManager->DoesEntityHaveElement(Entity, FTestFragment_Int::StaticStruct()));
+			AITEST_FALSE(TEXT("Entity should no longer have FTestFragment_Float"), EntityManager->DoesEntityHaveElement(Entity, FTestFragment_Float::StaticStruct()));
+			AITEST_FALSE(TEXT("Entity should no longer have TagA"), EntityManager->DoesEntityHaveElement(Entity, FTestTag_A::StaticStruct()));
+			AITEST_FALSE(TEXT("Entity should no longer have const shared"), EntityManager->DoesEntityHaveElement(Entity, FTestConstSharedFragment_Int::StaticStruct()));
+			AITEST_FALSE(TEXT("Entity should no longer have shared"), EntityManager->DoesEntityHaveElement(Entity, FTestSharedFragment_Int::StaticStruct()));
+		}
+
+		return true;
+	}
+};
+IMPLEMENT_AI_INSTANT_TEST(FCommands_RemoveElement_MixedTypes, "System.Mass.Commands.RemoveElement.MixedTypes");
+
+struct FCommands_AddRemoveElement_MixedTypes : FEntityTestBase
+{
+	virtual bool InstantTest() override
+	{
+		constexpr int32 Count = 5;
+		const FMassArchetypeHandle IntsTagAArchetype = EntityManager->CreateArchetype({ FTestFragment_Int::StaticStruct(), FTestTag_A::StaticStruct() });
+
+		TArray<FMassEntityHandle> Entities;
+		EntityManager->BatchCreateEntities(IntsTagAArchetype, Count, Entities);
+
+		// Add a const shared fragment to all entities
+		FTestConstSharedFragment_Int SharedFragmentInstance(42);
+		FConstSharedStruct ConstSharedStruct = EntityManager->GetOrCreateConstSharedFragment(SharedFragmentInstance);
+		for (const FMassEntityHandle& Entity : Entities)
+		{
+			EntityManager->AddConstSharedFragmentToEntity(Entity, ConstSharedStruct);
+		}
+
+		// Add a fragment and a tag via AddElement
+		FMassCommandAddElement& AddFragmentCmd = EntityManager->Defer().PushUniqueCommand<FMassCommandAddElement>(FTestFragment_Float::StaticStruct());
+		FMassCommandAddElement& AddTagCmd = EntityManager->Defer().PushUniqueCommand<FMassCommandAddElement>(FTestTag_B::StaticStruct());
+		// Remove a tag and a const shared fragment via RemoveElement
+		FMassCommandRemoveElement& RemoveTagCmd = EntityManager->Defer().PushUniqueCommand<FMassCommandRemoveElement>(FTestTag_A::StaticStruct());
+		FMassCommandRemoveElement& RemoveSharedCmd = EntityManager->Defer().PushUniqueCommand<FMassCommandRemoveElement>(FTestConstSharedFragment_Int::StaticStruct());
+
+		for (const FMassEntityHandle& Entity : Entities)
+		{
+			AddFragmentCmd.Add(Entity);
+			AddTagCmd.Add(Entity);
+			RemoveTagCmd.Add(Entity);
+			RemoveSharedCmd.Add(Entity);
+		}
+
+		EntityManager->FlushCommands();
+
+		for (const FMassEntityHandle& Entity : Entities)
+		{
+			AITEST_TRUE(TEXT("Entity should have FTestFragment_Int"), EntityManager->DoesEntityHaveElement(Entity, FTestFragment_Int::StaticStruct()));
+			AITEST_TRUE(TEXT("Entity should have FTestFragment_Float"), EntityManager->DoesEntityHaveElement(Entity, FTestFragment_Float::StaticStruct()));
+			AITEST_TRUE(TEXT("Entity should have TagB"), EntityManager->DoesEntityHaveElement(Entity, FTestTag_B::StaticStruct()));
+			AITEST_FALSE(TEXT("Entity should no longer have TagA"), EntityManager->DoesEntityHaveElement(Entity, FTestTag_A::StaticStruct()));
+			AITEST_FALSE(TEXT("Entity should no longer have const shared fragment"), EntityManager->DoesEntityHaveElement(Entity, FTestConstSharedFragment_Int::StaticStruct()));
+		}
+
+		return true;
+	}
+};
+IMPLEMENT_AI_INSTANT_TEST(FCommands_AddRemoveElement_MixedTypes, "System.Mass.Commands.AddRemoveElement.MixedTypes");
+
+struct FCommands_RemoveElement_OnlySharedFragments : FEntityTestBase
+{
+	virtual bool InstantTest() override
+	{
+		constexpr int32 Count = 5;
+		TArray<FMassEntityHandle> Entities;
+		EntityManager->BatchCreateEntities(IntsArchetype, Count, Entities);
+
+		// Add a mutable shared fragment to all entities
+		FTestSharedFragment_Int SharedFragmentInstance(42);
+		FSharedStruct SharedStruct = FSharedStruct::Make(SharedFragmentInstance);
+		FMassArchetypeSharedFragmentValues SharedFragmentValues;
+		SharedFragmentValues.Add(SharedStruct);
+
+		TArray<FMassArchetypeEntityCollection> EntityCollections;
+		UE::Mass::Utils::CreateEntityCollections(*EntityManager.Get(), Entities, FMassArchetypeEntityCollection::NoDuplicates, EntityCollections);
+		EntityManager->BatchAddSharedFragmentsForEntities(EntityCollections, SharedFragmentValues);
+
+		for (const FMassEntityHandle& Entity : Entities)
+		{
+			AITEST_TRUE(TEXT("Entity should have shared fragment before removal"), EntityManager->DoesEntityHaveElement(Entity, FTestSharedFragment_Int::StaticStruct()));
+		}
+
+		FMassCommandRemoveElement& Command = EntityManager->Defer().PushUniqueCommand<FMassCommandRemoveElement>(FTestSharedFragment_Int::StaticStruct());
+		for (const FMassEntityHandle& Entity : Entities)
+		{
+			Command.Add(Entity);
+		}
+
+		EntityManager->FlushCommands();
+
+		for (const FMassEntityHandle& Entity : Entities)
+		{
+			AITEST_FALSE(TEXT("Entity should no longer have shared fragment"), EntityManager->DoesEntityHaveElement(Entity, FTestSharedFragment_Int::StaticStruct()));
+			AITEST_TRUE(TEXT("Entity should still have FTestFragment_Int"), EntityManager->DoesEntityHaveElement(Entity, FTestFragment_Int::StaticStruct()));
+		}
+
+		return true;
+	}
+};
+IMPLEMENT_AI_INSTANT_TEST(FCommands_RemoveElement_OnlySharedFragments, "System.Mass.Commands.RemoveElement.OnlySharedFragments");
+
+struct FCommands_AddElement_RejectsSharedFragment : FEntityTestBase
+{
+	virtual bool InstantTest() override
+	{
+		constexpr int32 Count = 1;
+		TArray<FMassEntityHandle> Entities;
+		EntityManager->BatchCreateEntities(IntsArchetype, Count, Entities);
+
+		FMassCommandAddElement& Command = EntityManager->Defer().PushUniqueCommand<FMassCommandAddElement>(FTestSharedFragment_Int::StaticStruct());
+		Command.Add(Entities[0]);
+
+		{
+			MASS_SCOPED_ENSURE_TEST(TEXT("Only tags, fragments, and sparse elements can be added via this API"), Count);
+			EntityManager->FlushCommands();
+		}
+
+		return true;
+	}
+};
+IMPLEMENT_AI_INSTANT_TEST(FCommands_AddElement_RejectsSharedFragment, "System.Mass.Commands.AddElement.RejectsSharedFragment");
+
+struct FCommands_AddElement_RejectsConstSharedFragment : FEntityTestBase
+{
+	virtual bool InstantTest() override
+	{
+		constexpr int32 Count = 1;
+		TArray<FMassEntityHandle> Entities;
+		EntityManager->BatchCreateEntities(IntsArchetype, Count, Entities);
+
+		FMassCommandAddElement& Command = EntityManager->Defer().PushUniqueCommand<FMassCommandAddElement>(FTestConstSharedFragment_Int::StaticStruct());
+		Command.Add(Entities[0]);
+
+		{
+			MASS_SCOPED_ENSURE_TEST(TEXT("Only tags, fragments, and sparse elements can be added via this API"), Count);
+			EntityManager->FlushCommands();
+		}
+
+		return true;
+	}
+};
+IMPLEMENT_AI_INSTANT_TEST(FCommands_AddElement_RejectsConstSharedFragment, "System.Mass.Commands.AddElement.RejectsConstSharedFragment");
+
 #endif // WITH_MASSENTITY_DEBUG
 } // FMassCommandsTest
 
diff --git a/Engine/Source/Runtime/MassEntity/Private/MassEntityManager.cpp b/Engine/Source/Runtime/MassEntity/Private/MassEntityManager.cpp
--- a/Engine/Source/Runtime/MassEntity/Private/MassEntityManager.cpp
+++ b/Engine/Source/Runtime/MassEntity/Private/MassEntityManager.cpp
@@ -54,7 +54,9 @@ namespace UE::Mass::Private
 	bool IsElement(TNotNull<const UScriptStruct*> ElementType)
 	{
 		return UE::Mass::IsA<FMassFragment>(ElementType)
-			|| UE::Mass::IsA<FMassTag>(ElementType);
+			|| UE::Mass::IsA<FMassTag>(ElementType)
+			|| UE::Mass::IsA<FMassSharedFragment>(ElementType)
+			|| UE::Mass::IsA<FMassConstSharedFragment>(ElementType);
 	}
 
 	bool IsSparseElement(TNotNull<const UScriptStruct*> ElementType)
@@ -1513,7 +1515,18 @@ bool FMassEntityManager::DoesEntityHaveElement(FMassEntityHandle Entity, TNotNul
 
 void FMassEntityManager::BatchAddElementToEntities(TConstArrayView<FMassEntityHandle> Entities, TNotNull<const UScriptStruct*> ElementType)
 {
-	CHECK_ELEMENT(ElementType);
+	const bool bIsAddableElement = UE::Mass::Private::IsElement(ElementType)
+		&& !UE::Mass::IsA<FMassSharedFragment>(ElementType)
+		&& !UE::Mass::IsA<FMassConstSharedFragment>(ElementType);
+	testableCheckf(bIsAddableElement
+		, TEXT("%hs: Only tags, fragments, and sparse elements can be added via this API."
+			" Shared fragments require a value - use AddSharedFragmentToEntity"
+			" or AddConstSharedFragmentToEntity instead. Type provided: %s")
+		, __FUNCTION__, *ElementType->GetName());
+	if (!bIsAddableElement)
+	{
+		return;
+	}
 
 	TArray<FMassArchetypeEntityCollection> EntityCollections;
 	UE::Mass::Utils::CreateEntityCollections(*this, Entities, FMassArchetypeEntityCollection::EDuplicatesHandling::FoldDuplicates, EntityCollections);
@@ -1534,14 +1547,25 @@ void FMassEntityManager::BatchAddElementToEntities(TConstArrayView<FMassEntityHa
 
 void FMassEntityManager::AddElementToEntity(FMassEntityHandle EntityHandle, TNotNull<const UScriptStruct*> ElementType)
 {
-	CHECK_ELEMENT(ElementType);
+	const bool bIsAddableElement = UE::Mass::Private::IsElement(ElementType)
+		&& !UE::Mass::IsA<FMassSharedFragment>(ElementType)
+		&& !UE::Mass::IsA<FMassConstSharedFragment>(ElementType);
+	testableCheckf(bIsAddableElement
+		, TEXT("%hs: Only tags, fragments, and sparse elements can be added via this API."
+			" Shared fragments require a value - use AddSharedFragmentToEntity"
+			" or AddConstSharedFragmentToEntity instead. Type provided: %s")
+		, __FUNCTION__, *ElementType->GetName());
+	if (!bIsAddableElement)
+	{
+		return;
+	}
 
 	if (UE::Mass::IsSparse(ElementType))
 	{
 		TSharedPtr<FMassArchetypeData>& CurrentArchetype = GetEntityStorageInterface().GetArchetypeAsShared(EntityHandle.Index);
 		checkf(CurrentArchetype, TEXT("Entity (ID: %d, SN:%d) not assigned an archetype yet"), EntityHandle.Index, EntityHandle.SerialNumber);
 
-		FStructView _ = CurrentArchetype->AddSparseElementToEntity(EntityHandle, ElementType);	
+		FStructView _ = CurrentArchetype->AddSparseElementToEntity(EntityHandle, ElementType);
 	}
 	else if (UE::Mass::IsA<FMassTag>(ElementType))
 	{
@@ -1567,6 +1591,14 @@ void FMassEntityManager::BatchRemoveElementFromEntities(TConstArrayView<FMassEnt
 	{
 		BatchChangeTagsForEntities(EntityCollections, {}, FMassTagBitSet(ElementType));
 	}
+	else if (UE::Mass::IsA<FMassSharedFragment>(ElementType))
+	{
+		BatchRemoveSharedFragmentFromEntities(EntityCollections, FMassSharedFragmentBitSet(ElementType));
+	}
+	else if (UE::Mass::IsA<FMassConstSharedFragment>(ElementType))
+	{
+		BatchRemoveConstSharedFragmentFromEntities(EntityCollections, FMassConstSharedFragmentBitSet(ElementType));
+	}
 	else
 	{
 		BatchChangeFragmentCompositionForEntities(EntityCollections, {}, FMassFragmentBitSet(ElementType));
@@ -1585,6 +1617,14 @@ void FMassEntityManager::RemoveElementFromEntity(FMassEntityHandle Entity, TNotN
 	{
 		RemoveTagFromEntity(Entity, ElementType);
 	}
+	else if (UE::Mass::IsA<FMassSharedFragment>(ElementType))
+	{
+		RemoveSharedFragmentFromEntity(Entity, ElementType);
+	}
+	else if (UE::Mass::IsA<FMassConstSharedFragment>(ElementType))
+	{
+		RemoveConstSharedFragmentFromEntity(Entity, ElementType);
+	}
 	else
 	{
 		RemoveFragmentFromEntity(Entity, ElementType);
@@ -2060,6 +2100,81 @@ void FMassEntityManager::BatchRemoveSparseElementFromEntities(TConstArrayView<FM
 	}
 }
 
+void FMassEntityManager::BatchRemoveSharedFragmentFromEntities(
+	TConstArrayView<FMassArchetypeEntityCollection> EntityCollections,
+	const FMassSharedFragmentBitSet& ToRemove)
+{
+	TRACE_CPUPROFILER_EVENT_SCOPE(Mass_BatchRemoveSharedFragmentFromEntities);
+	InternalBatchRemoveSharedFragmentsFromEntities(EntityCollections, FMassElementBitSet(ToRemove)
+		, &ToRemove, /*ConstSharedFragmentToRemoveBitSet=*/nullptr);
+}
+
+void FMassEntityManager::BatchRemoveConstSharedFragmentFromEntities(
+	TConstArrayView<FMassArchetypeEntityCollection> EntityCollections,
+	const FMassConstSharedFragmentBitSet& ToRemove)
+{
+	TRACE_CPUPROFILER_EVENT_SCOPE(Mass_BatchRemoveConstSharedFragmentFromEntities);
+	InternalBatchRemoveSharedFragmentsFromEntities(EntityCollections, FMassElementBitSet(ToRemove)
+		, /*SharedFragmentToRemoveBitSet=*/nullptr, &ToRemove);
+}
+
+void FMassEntityManager::InternalBatchRemoveSharedFragmentsFromEntities(
+	TConstArrayView<FMassArchetypeEntityCollection> EntityCollections,
+	const FMassElementBitSet& ElementsToRemoveBitSet,
+	const FMassSharedFragmentBitSet* SharedFragmentToRemoveBitSet,
+	const FMassConstSharedFragmentBitSet* ConstSharedFragmentToRemoveBitSet)
+{
+	CHECK_SYNC_API();
+
+	const bool bElementsRemovedAreObserved = ObserverManager.HasObserversForBitSet(
+		ElementsToRemoveBitSet, EMassObservedOperation::RemoveElement);
+
+	for (const FMassArchetypeEntityCollection& Collection : EntityCollections)
+	{
+		FMassArchetypeData* CurrentArchetype = Collection.GetArchetype().DataPtr.Get();
+		if (!CurrentArchetype)
+		{
+			continue;
+		}
+
+		if (!CurrentArchetype->GetCompositionBitSet().HasAny(ElementsToRemoveBitSet))
+		{
+			continue;
+		}
+
+		if (bElementsRemovedAreObserved)
+		{
+			FMassElementBitSet ElementsRemoved = ElementsToRemoveBitSet.GetOverlap(CurrentArchetype->GetCompositionBitSet());
+			if (ElementsRemoved.IsEmpty() == false)
+			{
+				ObserverManager.OnCompositionChanged(Collection, MoveTemp(ElementsRemoved), EMassObservedOperation::RemoveElement);
+			}
+		}
+
+		FMassArchetypeCompositionDescriptor NewComposition(CurrentArchetype->GetCompositionDescriptor());
+		NewComposition.GetElementsBitSet().Remove(ElementsToRemoveBitSet);
+		const FMassArchetypeHandle NewArchetypeHandle = CreateArchetype(NewComposition);
+		check(NewArchetypeHandle.IsValid());
+		FMassArchetypeData* NewArchetype = NewArchetypeHandle.DataPtr.Get();
+		check(NewArchetype);
+
+		NewArchetype->CopyDebugNamesFrom(*CurrentArchetype);
+
+		TArray<FMassEntityHandle> EntitiesBeingMoved;
+		CurrentArchetype->BatchMoveEntitiesToAnotherArchetype(Collection,
+			*NewArchetype,
+			EntitiesBeingMoved, /*OutNewRanges=*/nullptr,
+			/*SharedFragmentValuesToAdd=*/nullptr,
+			SharedFragmentToRemoveBitSet, ConstSharedFragmentToRemoveBitSet);
+
+		for (const FMassEntityHandle& Entity : EntitiesBeingMoved)
+		{
+			check(GetEntityStorageInterface().IsValidIndex(Entity.Index));
+			GetEntityStorageInterface().SetArchetypeFromShared(Entity.Index, NewArchetypeHandle.DataPtr);
+		}
+	}
+}
+
 void FMassEntityManager::MoveEntityToAnotherArchetype(FMassEntityHandle Entity, FMassArchetypeHandle NewArchetypeHandle,
 	const FMassArchetypeSharedFragmentValues* SharedFragmentValuesOverride)
 {
diff --git a/Engine/Source/Runtime/MassEntity/Public/MassEntityManager.h b/Engine/Source/Runtime/MassEntity/Public/MassEntityManager.h
--- a/Engine/Source/Runtime/MassEntity/Public/MassEntityManager.h
+++ b/Engine/Source/Runtime/MassEntity/Public/MassEntityManager.h
@@ -609,6 +609,22 @@ public:
 	// @todo add a version with FMassArchetypeEntityCollectionWithPayload to add pre-configured fragment instances
 	UE_API void BatchRemoveSparseElementFromEntities(TConstArrayView<FMassArchetypeEntityCollection> EntityCollections, TNotNull<const UScriptStruct*> ElementType);
 
+	/**
+	 * Removes the specified shared fragment types from all entities in the provided collections.
+	 * Entities are moved to new archetypes without the removed shared fragment types.
+	 * @param EntityCollections the entity collections to process
+	 * @param ToRemove bitset of shared fragment types to remove
+	 */
+	UE_API void BatchRemoveSharedFragmentFromEntities(TConstArrayView<FMassArchetypeEntityCollection> EntityCollections, const FMassSharedFragmentBitSet& ToRemove);
+
+	/**
+	 * Removes the specified const shared fragment types from all entities in the provided collections.
+	 * Entities are moved to new archetypes without the removed shared fragment types.
+	 * @param EntityCollections the entity collections to process
+	 * @param ToRemove bitset of const shared fragment types to remove
+	 */
+	UE_API void BatchRemoveConstSharedFragmentFromEntities(TConstArrayView<FMassArchetypeEntityCollection> EntityCollections, const FMassConstSharedFragmentBitSet& ToRemove);
+
 	/**
 	 * Adds elements indicated by InOutDescriptor to the entity indicated by EntityHandle. The function also figures out which elements
 	 * in InOutDescriptor are missing from the current composition of the given entity and then returns the resulting 
@@ -1345,6 +1361,12 @@ protected:
 		TConstArrayView<const UScriptStruct*> FragmentsAndTagsList) const;
 
 private:
+	void InternalBatchRemoveSharedFragmentsFromEntities(
+		TConstArrayView<FMassArchetypeEntityCollection> EntityCollections,
+		const FMassElementBitSet& ElementsToRemoveBitSet,
+		const FMassSharedFragmentBitSet* SharedFragmentToRemoveBitSet,
+		const FMassConstSharedFragmentBitSet* ConstSharedFragmentToRemoveBitSet);
+
 	void InternalBuildEntity(FMassEntityHandle EntityHandle, const FMassArchetypeHandle& ArchetypeHandle, const FMassArchetypeSharedFragmentValues& SharedFragmentValues);
 	void InternalReleaseEntity(FMassEntityHandle EntityHandle);
 
