[Mass] Add HasParticipants, HasRelationSubjects, HasRelationObjects to FRelationManager

Zero-allocation existence checks for relation participants.
HasParticipants short-circuits on first valid entity handle.
HasRelationSubjects/HasRelationObjects are convenience wrappers.

#Synced-CL 52461421

diff --git a/Engine/Source/Developer/MassEntityTestSuite/Private/MassRelationTest.cpp b/Engine/Source/Developer/MassEntityTestSuite/Private/MassRelationTest.cpp
--- a/Engine/Source/Developer/MassEntityTestSuite/Private/MassRelationTest.cpp
+++ b/Engine/Source/Developer/MassEntityTestSuite/Private/MassRelationTest.cpp
@@ -758,6 +758,140 @@ namespace UE::Mass::Test::Relations
 	};
 	IMPLEMENT_AI_INSTANT_TEST(FSymmetricCleanup, "System.Mass.Relations.SymmetricCleanup");
 
+	struct FHasRelationSubjectsAndObjects : FRelationTestBase
+	{
+		virtual bool InstantTest() override
+		{
+			const FMassArchetypeHandle OriginalArchetype = IntsArchetype;
+
+			TArray<FMassEntityHandle> Entities;
+			EntityManager->BatchCreateEntities(OriginalArchetype, 3, Entities);
+
+			const FMassEntityHandle ChildA = Entities[0];
+			const FMassEntityHandle ChildB = Entities[1];
+			const FMassEntityHandle Parent = Entities[2];
+
+			// Before any relations: nothing should exist
+			AITEST_FALSE("Parent has no subjects before relations", RelationshipManagerPtr->HasRelationSubjects(RelationTypeHandle, Parent));
+			AITEST_FALSE("ChildA has no objects before relations", RelationshipManagerPtr->HasRelationObjects(RelationTypeHandle, ChildA));
+			AITEST_FALSE("ChildA has no subjects before relations", RelationshipManagerPtr->HasRelationSubjects(RelationTypeHandle, ChildA));
+
+			// Create ChildA -> Parent
+			RelationshipManagerPtr->CreateRelationInstance(RelationTypeHandle, ChildA, Parent);
+			EntityManager->FlushCommands();
+
+			AITEST_TRUE("Parent has subjects after first relation", RelationshipManagerPtr->HasRelationSubjects(RelationTypeHandle, Parent));
+			AITEST_TRUE("ChildA has objects after first relation", RelationshipManagerPtr->HasRelationObjects(RelationTypeHandle, ChildA));
+			AITEST_FALSE("ChildA has no subjects (it is a leaf)", RelationshipManagerPtr->HasRelationSubjects(RelationTypeHandle, ChildA));
+			AITEST_FALSE("Parent has no objects (it is a root)", RelationshipManagerPtr->HasRelationObjects(RelationTypeHandle, Parent));
+
+			// Create ChildB -> Parent (second child)
+			RelationshipManagerPtr->CreateRelationInstance(RelationTypeHandle, ChildB, Parent);
+			EntityManager->FlushCommands();
+
+			AITEST_TRUE("Parent still has subjects after second relation", RelationshipManagerPtr->HasRelationSubjects(RelationTypeHandle, Parent));
+			AITEST_TRUE("ChildB has objects after relation", RelationshipManagerPtr->HasRelationObjects(RelationTypeHandle, ChildB));
+
+			// Destroy ChildA -> Parent
+			RelationshipManagerPtr->DestroyRelationInstance(RelationTypeHandle, ChildA, Parent);
+			EntityManager->FlushCommands();
+
+			AITEST_TRUE("Parent still has subjects (ChildB remains)", RelationshipManagerPtr->HasRelationSubjects(RelationTypeHandle, Parent));
+			AITEST_FALSE("ChildA has no objects after relation destroyed", RelationshipManagerPtr->HasRelationObjects(RelationTypeHandle, ChildA));
+
+			// Destroy ChildB -> Parent
+			RelationshipManagerPtr->DestroyRelationInstance(RelationTypeHandle, ChildB, Parent);
+			EntityManager->FlushCommands();
+
+			AITEST_FALSE("Parent has no subjects after all relations destroyed", RelationshipManagerPtr->HasRelationSubjects(RelationTypeHandle, Parent));
+			AITEST_FALSE("ChildB has no objects after relation destroyed", RelationshipManagerPtr->HasRelationObjects(RelationTypeHandle, ChildB));
+
+			return true;
+		}
+	};
+	IMPLEMENT_AI_INSTANT_TEST(FHasRelationSubjectsAndObjects, "System.Mass.Relations.HasRelationSubjectsAndObjects");
+
+	struct FForEachRelationSubjectAndObject : FRelationTestBase
+	{
+		virtual bool InstantTest() override
+		{
+			const FMassArchetypeHandle OriginalArchetype = IntsArchetype;
+
+			TArray<FMassEntityHandle> Entities;
+			EntityManager->BatchCreateEntities(OriginalArchetype, 4, Entities);
+
+			const FMassEntityHandle ChildA = Entities[0];
+			const FMassEntityHandle ChildB = Entities[1];
+			const FMassEntityHandle ChildC = Entities[2];
+			const FMassEntityHandle Parent = Entities[3];
+
+			RelationshipManagerPtr->CreateRelationInstance(RelationTypeHandle, ChildA, Parent);
+			RelationshipManagerPtr->CreateRelationInstance(RelationTypeHandle, ChildB, Parent);
+			RelationshipManagerPtr->CreateRelationInstance(RelationTypeHandle, ChildC, Parent);
+			EntityManager->FlushCommands();
+
+			// ForEachRelationSubject: enumerate children of Parent
+			{
+				TArray<FMassEntityHandle> CollectedSubjects;
+				RelationshipManagerPtr->ForEachRelationSubject(RelationTypeHandle, Parent,
+					[&CollectedSubjects](FMassEntityHandle Entity)
+					{
+						CollectedSubjects.Add(Entity);
+					});
+
+				TArray<FMassEntityHandle> ExpectedSubjects = RelationshipManagerPtr->GetRelationSubjects(RelationTypeHandle, Parent);
+				AITEST_EQUAL("ForEachRelationSubject count matches GetRelationSubjects", CollectedSubjects.Num(), ExpectedSubjects.Num());
+
+				// Verify same set of entities (order may differ)
+				for (const FMassEntityHandle& Expected : ExpectedSubjects)
+				{
+					AITEST_TRUE("ForEachRelationSubject contains expected entity", CollectedSubjects.Contains(Expected));
+				}
+			}
+
+			// ForEachRelationObject: enumerate parents of ChildA
+			{
+				TArray<FMassEntityHandle> CollectedObjects;
+				RelationshipManagerPtr->ForEachRelationObject(RelationTypeHandle, ChildA,
+					[&CollectedObjects](FMassEntityHandle Entity)
+					{
+						CollectedObjects.Add(Entity);
+					});
+
+				TArray<FMassEntityHandle> ExpectedObjects = RelationshipManagerPtr->GetRelationObjects(RelationTypeHandle, ChildA);
+				AITEST_EQUAL("ForEachRelationObject count matches GetRelationObjects", CollectedObjects.Num(), ExpectedObjects.Num());
+
+				for (const FMassEntityHandle& Expected : ExpectedObjects)
+				{
+					AITEST_TRUE("ForEachRelationObject contains expected entity", CollectedObjects.Contains(Expected));
+				}
+			}
+
+			// ForEach on entity with no relations: callback should not fire
+			{
+				int32 CallCount = 0;
+				RelationshipManagerPtr->ForEachRelationSubject(RelationTypeHandle, ChildA,
+					[&CallCount](FMassEntityHandle)
+					{
+						++CallCount;
+					});
+				AITEST_EQUAL("ForEachRelationSubject on leaf fires zero callbacks", CallCount, 0);
+			}
+			{
+				int32 CallCount = 0;
+				RelationshipManagerPtr->ForEachRelationObject(RelationTypeHandle, Parent,
+					[&CallCount](FMassEntityHandle)
+					{
+						++CallCount;
+					});
+				AITEST_EQUAL("ForEachRelationObject on root fires zero callbacks", CallCount, 0);
+			}
+
+			return true;
+		}
+	};
+	IMPLEMENT_AI_INSTANT_TEST(FForEachRelationSubjectAndObject, "System.Mass.Relations.ForEachRelationSubjectAndObject");
+
 }// UE::Mass::Test::Relations
 
 UE_ENABLE_OPTIMIZATION_SHIP
diff --git a/Engine/Source/Runtime/MassEntity/Private/MassRelationManager.cpp b/Engine/Source/Runtime/MassEntity/Private/MassRelationManager.cpp
--- a/Engine/Source/Runtime/MassEntity/Private/MassRelationManager.cpp
+++ b/Engine/Source/Runtime/MassEntity/Private/MassRelationManager.cpp
@@ -114,6 +114,23 @@ namespace UE::Mass
 		}
 	}
 
+	bool FRelationData::HasParticipants(const FMassEntityManager& EntityManager, const FMassEntityHandle RoleEntity,
+		const ERelationRole QueriedRole) const
+	{
+		const int32 QueriedRoleIndex = static_cast<int32>(QueriedRole);
+		if (const FRelationInstanceData* InstanceData = RoleMap.Find(RoleEntity))
+		{
+			for (const FMassRelationRoleInstanceHandle RelationInstanceHandle : (*InstanceData)[QueriedRoleIndex])
+			{
+				if (RelationInstanceHandle.GetRoleEntityHandle(EntityManager).IsValid())
+				{
+					return true;
+				}
+			}
+		}
+		return false;
+	}
+
 	//-----------------------------------------------------------------------------
 	// FRelationManager::FHierarchyEntitiesContainer
 	//-----------------------------------------------------------------------------
@@ -398,6 +415,18 @@ namespace UE::Mass
 		RelationData.ForEachParticipant(EntityManager, SubjectEntity, ERelationRole::Object, Callback);
 	}
 
+	bool FRelationManager::HasRelationSubjects(const FTypeHandle RelationTypeHandle, const FMassEntityHandle ObjectEntity) const
+	{
+		const FRelationData& RelationData = GetRelationDataChecked(RelationTypeHandle);
+		return RelationData.HasParticipants(EntityManager, ObjectEntity, ERelationRole::Subject);
+	}
+
+	bool FRelationManager::HasRelationObjects(const FTypeHandle RelationTypeHandle, const FMassEntityHandle SubjectEntity) const
+	{
+		const FRelationData& RelationData = GetRelationDataChecked(RelationTypeHandle);
+		return RelationData.HasParticipants(EntityManager, SubjectEntity, ERelationRole::Object);
+	}
+
 	void FRelationManager::GatherHierarchy(const FRelationData& RelationData, const FMassEntityHandle SubjectHandle, FHierarchyEntitiesContainer& SubSubjects, const uint32 Depth) const
 	{
 		// getting subject's subjects
diff --git a/Engine/Source/Runtime/MassEntity/Public/MassRelationManager.h b/Engine/Source/Runtime/MassEntity/Public/MassRelationManager.h
--- a/Engine/Source/Runtime/MassEntity/Public/MassRelationManager.h
+++ b/Engine/Source/Runtime/MassEntity/Public/MassRelationManager.h
@@ -62,6 +62,9 @@ namespace UE::Mass
 
 		TArray<FMassEntityHandle> GetParticipants(const FMassEntityManager& EntityManager, const FMassEntityHandle RoleEntity, ERelationRole QueriedRole) const;
 
+		/** Returns true if RoleEntity has at least one valid participant in the queried role. No allocation. */
+		bool HasParticipants(const FMassEntityManager& EntityManager, FMassEntityHandle RoleEntity, ERelationRole QueriedRole) const;
+
 		/** Callback-based enumeration of participants. Same iteration as GetParticipants but invokes a callback per valid entity instead of building a TArray. */
 		void ForEachParticipant(const FMassEntityManager& EntityManager, FMassEntityHandle RoleEntity,
 			ERelationRole QueriedRole, TFunctionRef<void(FMassEntityHandle)> Callback) const;
@@ -112,6 +115,12 @@ namespace UE::Mass
 		UE_API void ForEachRelationObject(const FTypeHandle RelationTypeHandle, FMassEntityHandle SubjectEntity,
 			TFunctionRef<void(FMassEntityHandle)> Callback) const;
 
+		/** Returns true if ObjectEntity has at least one subject in this relation type. No allocation. */
+		UE_API bool HasRelationSubjects(const FTypeHandle RelationTypeHandle, FMassEntityHandle ObjectEntity) const;
+
+		/** Returns true if SubjectEntity has at least one object in this relation type. No allocation. */
+		UE_API bool HasRelationObjects(const FTypeHandle RelationTypeHandle, FMassEntityHandle SubjectEntity) const;
+
 		UE_API bool IsSubjectOfRelation(const FTypeHandle RelationTypeHandle, const FMassEntityHandle Subject, const FMassEntityHandle Object) const;
 		UE_API bool IsSubjectOfRelation(const FRelationData& RelationDataInstance, const FMassEntityHandle Subject, const FMassEntityHandle Object) const;
 		UE_API bool IsSubjectOfRelationRecursive(const FTypeHandle RelationTypeHandle, const FMassEntityHandle Subject, const FMassEntityHandle Object) const;
