[Mass] Fixed matches queries with Any fragments and tags to archetypes with only one or the other.

When Any is used in a query, it's expected to match the archetype if at least one Any element matches, irrespective if those are fragments or tags. The existing test however wouldn't match correctly if there's an Any fragment and an Any tag and the archetype only has one or the other.

#Synced-CL 52452101

diff --git a/Engine/Source/Developer/MassEntityTestSuite/Private/StressTests/MassStressRequirementsTest.cpp b/Engine/Source/Developer/MassEntityTestSuite/Private/StressTests/MassStressRequirementsTest.cpp
--- a/Engine/Source/Developer/MassEntityTestSuite/Private/StressTests/MassStressRequirementsTest.cpp
+++ b/Engine/Source/Developer/MassEntityTestSuite/Private/StressTests/MassStressRequirementsTest.cpp
@@ -929,6 +929,39 @@ struct FRequirements_DoesArchetypeMatchRequirements_Tags : FEntityTestBase
 };
 IMPLEMENT_AI_INSTANT_TEST(FRequirements_DoesArchetypeMatchRequirements_Tags, "System.Mass.Stress.Requirements.Match.Tags");
 
+struct FRequirements_DoesArchetypeMatchRequirements_AnyFragmentsAndTags : FEntityTestBase
+{
+	virtual bool InstantTest() override
+	{
+		// Query has FTestFragment_Int and FTestTag_A both as Any.
+		// An archetype should match if it has either one, not necessarily both.
+		FMassFragmentRequirements Requirements(EntityManager);
+		Requirements.AddRequirement<FTestFragment_Int>(EMassFragmentAccess::ReadOnly, EMassFragmentPresence::Any);
+		Requirements.AddTagRequirement<FTestTag_A>(EMassFragmentPresence::Any);
+		Requirements.CheckValidity();
+
+		FMassArchetypeCompositionDescriptor FragmentOnlyComposition;
+		FragmentOnlyComposition.Add<FTestFragment_Int>();
+		FMassArchetypeHandle FragmentOnlyArchetype = EntityManager->CreateArchetype(FragmentOnlyComposition);
+
+		FMassArchetypeCompositionDescriptor TagOnlyComposition;
+		TagOnlyComposition.Add<FTestTag_A>();
+		FMassArchetypeHandle TagOnlyArchetype = EntityManager->CreateArchetype(TagOnlyComposition);
+
+		AITEST_TRUE("Archetype with matching fragment but no matching tags should match",
+			Requirements.DoesArchetypeMatchRequirements(FragmentOnlyArchetype));
+
+		AITEST_TRUE("Archetype with matching tag but no matching fragments should match",
+			Requirements.DoesArchetypeMatchRequirements(TagOnlyArchetype));
+
+		AITEST_FALSE("Archetype with neither matching fragment nor tag should not match",
+			Requirements.DoesArchetypeMatchRequirements(FloatsArchetype));
+
+		return true;
+	}
+};
+IMPLEMENT_AI_INSTANT_TEST(FRequirements_DoesArchetypeMatchRequirements_AnyFragmentsAndTags, "System.Mass.Stress.Requirements.Match.AnyFragmentsAndTags");
+
 
 //----------------------------------------------------------------------//
 // FMassFragmentRequirements - AddElementRequirement Tests
diff --git a/Engine/Source/Runtime/MassEntity/Private/MassRequirements.cpp b/Engine/Source/Runtime/MassEntity/Private/MassRequirements.cpp
--- a/Engine/Source/Runtime/MassEntity/Private/MassRequirements.cpp
+++ b/Engine/Source/Runtime/MassEntity/Private/MassRequirements.cpp
@@ -284,10 +284,10 @@ bool FMassFragmentRequirements::DoesArchetypeMatchRequirements(const FMassElemen
 		if (bHasPositiveRequirements)
 		{
 			return ArchetypeCompositionBitSet.HasAll(RequiredAllFragments)
-				&& (RequiredAnyFragments.IsEmpty() || ArchetypeCompositionBitSet.HasAny(RequiredAnyFragments))
 				&& ArchetypeCompositionBitSet.HasAll(RequiredAllTags)
-				&& (RequiredAnyTags.IsEmpty() || ArchetypeCompositionBitSet.HasAny(RequiredAnyTags))
-				&& ArchetypeCompositionBitSet.HasAll(RequiredAllChunkFragments | RequiredAllSharedFragments | RequiredAllConstSharedFragments);
+				&& ArchetypeCompositionBitSet.HasAll(RequiredAllChunkFragments | RequiredAllSharedFragments | RequiredAllConstSharedFragments)
+				&& ((RequiredAnyFragments.IsEmpty() && RequiredAnyTags.IsEmpty()) 
+					|| ArchetypeCompositionBitSet.HasAny(RequiredAnyFragments | RequiredAnyTags));
 		}
 		else if (bHasOptionalRequirements)
 		{
