[MassCommands]
Add bitset-level and runtime type validation for Mass element composition

- Shadow GetTypeBitSet<T>() in DECLARE_NEWTYPEBITSET to static_assert that T matches the expected element subtype, not just FMassElement (root cause fix for typed bitset accepting wrong element types)
- Add testableCheckf guards in BatchChangeCompositionForEntities to validate that ElementsToAdd and ElementsToRemove contain only tags and fragments, catching misuse before it reaches InitializeWithSimilar

#Synced-CL 52485477

diff --git a/Engine/Source/Developer/MassEntityTestSuite/Private/MassCommandsTest.cpp b/Engine/Source/Developer/MassEntityTestSuite/Private/MassCommandsTest.cpp
--- a/Engine/Source/Developer/MassEntityTestSuite/Private/MassCommandsTest.cpp
+++ b/Engine/Source/Developer/MassEntityTestSuite/Private/MassCommandsTest.cpp
@@ -720,7 +720,84 @@ struct FCommands_AddElement_RejectsConstSharedFragment : FEntityTestBase
 };
 IMPLEMENT_AI_INSTANT_TEST(FCommands_AddElement_RejectsConstSharedFragment, "System.Mass.Commands.AddElement.RejectsConstSharedFragment");
 
+struct FCommands_BatchChangeComposition_RejectsSharedAdd : FEntityTestBase
+{
+	virtual bool InstantTest() override
+	{
+		constexpr int32 Count = 1;
+		TArray<FMassEntityHandle> Entities;
+		EntityManager->BatchCreateEntities(IntsArchetype, Count, Entities);
+
+		TArray<FMassArchetypeEntityCollection> EntityCollections;
+		UE::Mass::Utils::CreateEntityCollections(*EntityManager.Get(), Entities, FMassArchetypeEntityCollection::NoDuplicates, EntityCollections);
+
+		FMassConstSharedFragmentBitSet ConstSharedBits(FTestConstSharedFragment_Int::StaticStruct());
+		const FMassElementBitSet SharedBits(ConstSharedBits);
+
+		{
+			MASS_SCOPED_ENSURE_TEST(TEXT("ElementsToAdd contains shared or sparse element types"), Count);
+			EntityManager->BatchChangeCompositionForEntities(EntityCollections, SharedBits, FMassElementBitSet());
+		}
+
+		return true;
+	}
+};
+IMPLEMENT_AI_INSTANT_TEST(FCommands_BatchChangeComposition_RejectsSharedAdd, "System.Mass.Commands.BatchChangeComposition.RejectsSharedAdd");
+
+struct FCommands_BatchChangeComposition_RejectsSharedRemove : FEntityTestBase
+{
+	virtual bool InstantTest() override
+	{
+		constexpr int32 Count = 1;
+		TArray<FMassEntityHandle> Entities;
+		EntityManager->BatchCreateEntities(IntsArchetype, Count, Entities);
+
+		TArray<FMassArchetypeEntityCollection> EntityCollections;
+		UE::Mass::Utils::CreateEntityCollections(*EntityManager.Get(), Entities, FMassArchetypeEntityCollection::NoDuplicates, EntityCollections);
+
+		FMassSharedFragmentBitSet MutableSharedBits(FTestSharedFragment_Int::StaticStruct());
+		const FMassElementBitSet SharedBits(MutableSharedBits);
+
+		{
+			MASS_SCOPED_ENSURE_TEST(TEXT("ElementsToRemove contains shared or sparse element types"), Count);
+			EntityManager->BatchChangeCompositionForEntities(EntityCollections, FMassElementBitSet(), SharedBits);
+		}
+
+		return true;
+	}
+};
+IMPLEMENT_AI_INSTANT_TEST(FCommands_BatchChangeComposition_RejectsSharedRemove, "System.Mass.Commands.BatchChangeComposition.RejectsSharedRemove");
+
 #endif // WITH_MASSENTITY_DEBUG
+
+//----------------------------------------------------------------------//
+// Compile-time validation tests for typed bitsets (DECLARE_NEWTYPEBITSET).
+// Uncomment any line below to verify it produces a static_assert error:
+//   "Type does not match the expected element type for this bitset."
+//----------------------------------------------------------------------//
+#if 0
+static void BitsetTypeValidation_ShouldNotCompile()
+{
+	// Tag in a fragment bitset
+	FMassFragmentBitSet FragBitSet;
+	FragBitSet += FMassFragmentBitSet::GetTypeBitSet<FTestTag_A>();
+
+	// Const shared fragment in a fragment bitset
+	FragBitSet += FMassFragmentBitSet::GetTypeBitSet<FTestConstSharedFragment_Int>();
+
+	// Fragment in a tag bitset
+	FMassTagBitSet TagBitSet;
+	TagBitSet += FMassTagBitSet::GetTypeBitSet<FTestFragment_Int>();
+
+	// PopulateBitSet path (what commands use internally)
+	FMassFragmentBitSet Result;
+	UE::Mass::TMultiTypeList<FTestConstSharedFragment_Int>::PopulateBitSet(Result);
+
+	// ConstructFragmentBitSet helper
+	auto BitSet = UE::Mass::Utils::ConstructFragmentBitSet<EMassCommandCheckTime::CompileTimeCheck, FTestTag_A>();
+}
+#endif
+
 } // FMassCommandsTest
 
 UE_ENABLE_OPTIMIZATION_SHIP
diff --git a/Engine/Source/Runtime/MassEntity/Private/MassEntityManager.cpp b/Engine/Source/Runtime/MassEntity/Private/MassEntityManager.cpp
--- a/Engine/Source/Runtime/MassEntity/Private/MassEntityManager.cpp
+++ b/Engine/Source/Runtime/MassEntity/Private/MassEntityManager.cpp
@@ -1886,6 +1886,15 @@ void FMassEntityManager::BatchChangeCompositionForEntities(TConstArrayView<FMass
 {
 	CHECK_SYNC_API();
 
+	testableCheckf(ElementsToAdd.ContainsOnlyTagsAndFragments()
+		, TEXT("%hs: ElementsToAdd contains shared or sparse element types."
+			" Only tags and fragments are supported. Use the dedicated shared fragment APIs instead.")
+		, __FUNCTION__);
+	testableCheckf(ElementsToRemove.ContainsOnlyTagsAndFragments()
+		, TEXT("%hs: ElementsToRemove contains shared or sparse element types."
+			" Only tags and fragments are supported. Use the dedicated shared fragment APIs instead.")
+		, __FUNCTION__);
+
 	TRACE_CPUPROFILER_EVENT_SCOPE(Mass_BatchChangeCompositionForEntities);
 
 	const bool bElementsAddedAreObserved = ObserverManager.HasObserversForBitSet(ElementsToAdd, EMassObservedOperation::AddElement);
diff --git a/Engine/Source/Runtime/MassEntity/Public/MassElement.h b/Engine/Source/Runtime/MassEntity/Public/MassElement.h
--- a/Engine/Source/Runtime/MassEntity/Public/MassElement.h
+++ b/Engine/Source/Runtime/MassEntity/Public/MassElement.h
@@ -200,6 +200,13 @@ namespace UE::Mass
 		ContainerTypeName(const FElementBitSetBase& OtherBitSet) \
 			: FElementBitSetBase(OtherBitSet) \
 		{} \
+		template<typename T> \
+		static const FElementBitSetBase& GetTypeBitSet() \
+		{ \
+			static_assert(TIsDerivedFrom<typename TRemoveReference<T>::Type, BaseType>::Value, \
+				"Type does not match the expected element type for this bitset."); \
+			return FElementBitSetBase::GetTypeBitSet<T>(); \
+		} \
 	}; \
 	template<> struct TElementTraits<BaseType> \
 	{ \
