[Mass] adding static type registration to FTypeManager - a way to register subsystem/shared-fragment types once globally so that every new FTypeManager instance picks them up during RegisterBuiltInTypes(). Two API flavors: RAII handle (auto-deregisters on destruction) and persistent (manual UnregisterStaticType). Meant for types whose lifetime spans multiple entity managers (engine subsystems, game instance subsystems) that previously had no clean registration path.

also fixing a pre-existing copy-paste bug in the shared fragment traits test assertion (was checking subsystem traits variable instead of the shared fragment one).

#Synced-CL 52465828

diff --git a/Engine/Source/Developer/MassEntityTestSuite/Private/MassTypeManagerTest.cpp b/Engine/Source/Developer/MassEntityTestSuite/Private/MassTypeManagerTest.cpp
--- a/Engine/Source/Developer/MassEntityTestSuite/Private/MassTypeManagerTest.cpp
+++ b/Engine/Source/Developer/MassEntityTestSuite/Private/MassTypeManagerTest.cpp
@@ -1,8 +1,11 @@
 // Copyright Epic Games, Inc. All Rights Reserved.
 
 #include "AITestsCommon.h"
+#include "Engine/Engine.h"
 #include "MassEntityTestTypes.h"
 #include "MassTypeManager.h"
+#include "MassExecutionContext.h"
+#include "MassSubsystemAccess.h"
 
 #define LOCTEXT_NAMESPACE "MassTest"
 
@@ -10,6 +13,20 @@ UE_DISABLE_OPTIMIZATION_SHIP
 
 namespace UE::Mass::Test
 {
+	/** Helper to create, initialize, and post-initialize a fresh FMassEntityManager. */
+	TSharedRef<FMassEntityManager> MakeInitializedEntityManager(const bool bCallPostInitialize = true)
+	{
+		TSharedRef<FMassEntityManager> Manager = MakeShareable(new FMassEntityManager(FAITestHelpers::GetWorld()));
+		FMassEntityManagerStorageInitParams InitParams;
+		InitParams.Emplace<FMassEntityManager_InitParams_Concurrent>();
+		Manager->Initialize(InitParams);
+		if (bCallPostInitialize)
+		{
+			Manager->PostInitialize();
+		}
+		return Manager;
+	}
+
 	struct FTypeManager_Subsystem : FEntityTestBase
 	{
 		virtual bool InstantTest() override
@@ -20,13 +37,349 @@ namespace UE::Mass::Test
 			AITEST_EQUAL("Subsystem ThreadSafeWrite value", TestSubsystemTraits.bThreadSafeWrite, TMassExternalSubsystemTraits<UMassTestWorldSubsystem>::ThreadSafeWrite)
 
 			const FSharedFragmentTypeTraits TestSharedFragmentTraits = FSharedFragmentTypeTraits::Make<FTestSharedFragment_Int>();
-			AITEST_EQUAL("Shared fragment bGameThreadOnly value", TestSubsystemTraits.bGameThreadOnly, TMassExternalSubsystemTraits<UMassTestWorldSubsystem>::GameThreadOnly)
+			AITEST_EQUAL("Shared fragment bGameThreadOnly value", TestSharedFragmentTraits.bGameThreadOnly, TMassSharedFragmentTraits<FTestSharedFragment_Int>::GameThreadOnly)
 
 			return true;
 		}
 	};
 	IMPLEMENT_AI_INSTANT_TEST(FTypeManager_Subsystem, "System.Mass.TypeManager.StaticSubsystem");
-} // UE::Mass::Test
+
+	//-------------------------------------------------------------------------
+	// Static type registration tests
+	//-------------------------------------------------------------------------
+
+	/** Verify that a statically registered subsystem type appears in a new TypeManager after RegisterBuiltInTypes. */
+	struct FTypeManager_StaticRegistration_Subsystem : FExecutionTestBase
+	{
+		virtual bool InstantTest() override
+		{
+			CA_ASSUME(EntityManager);
+
+			// Register a subsystem type statically via the RAII handle
+			FTypeManager::FStaticTypeRegistrationHandle Handle = FTypeManager::RegisterStaticType<UMassTestWorldSubsystem>();
+			AITEST_TRUE("Handle is valid after registration", Handle.IsValid());
+
+			// Create a fresh entity manager and initialize it — this triggers RegisterBuiltInTypes
+			TSharedRef<FMassEntityManager> TestEntityManager = MakeInitializedEntityManager();
+
+			// The statically registered type should be in the new TypeManager
+			const FTypeHandle TypeHandle = FTypeManager::MakeTypeHandle(UMassTestWorldSubsystem::StaticClass());
+			const FTypeInfo* TypeInfo = TestEntityManager->GetTypeManager().GetTypeInfo(TypeHandle);
+			AITEST_NOT_NULL("Statically registered subsystem type found in TypeManager", TypeInfo);
+			AITEST_NOT_NULL("TypeInfo has subsystem traits", TypeInfo ? TypeInfo->GetAsSystemTraits() : nullptr);
+
+			TestEntityManager->Deinitialize();
+			return true;
+		}
+	};
+	IMPLEMENT_AI_INSTANT_TEST(FTypeManager_StaticRegistration_Subsystem, "System.Mass.TypeManager.StaticRegistration.Subsystem");
+
+	/** Verify that a statically registered shared fragment type appears in a new TypeManager after RegisterBuiltInTypes. */
+	struct FTypeManager_StaticRegistration_SharedFragment : FExecutionTestBase
+	{
+		virtual bool InstantTest() override
+		{
+			CA_ASSUME(EntityManager);
+
+			FTypeManager::FStaticTypeRegistrationHandle Handle = FTypeManager::RegisterStaticType<FTestSharedFragment_Int>();
+			AITEST_TRUE("Handle is valid after registration", Handle.IsValid());
+
+			TSharedRef<FMassEntityManager> TestEntityManager = MakeInitializedEntityManager();
+
+			const FTypeHandle TypeHandle = FTypeManager::MakeTypeHandle(FTestSharedFragment_Int::StaticStruct());
+			const FTypeInfo* TypeInfo = TestEntityManager->GetTypeManager().GetTypeInfo(TypeHandle);
+			AITEST_NOT_NULL("Statically registered shared fragment type found in TypeManager", TypeInfo);
+			AITEST_NOT_NULL("TypeInfo has shared fragment traits", TypeInfo ? TypeInfo->GetAsSharedFragmentTraits() : nullptr);
+
+			TestEntityManager->Deinitialize();
+			return true;
+		}
+	};
+	IMPLEMENT_AI_INSTANT_TEST(FTypeManager_StaticRegistration_SharedFragment, "System.Mass.TypeManager.StaticRegistration.SharedFragment");
+
+	/** Verify that destroying the RAII handle removes the type from the static registry. */
+	struct FTypeManager_StaticRegistration_HandleDeregisters : FExecutionTestBase
+	{
+		virtual bool InstantTest() override
+		{
+			CA_ASSUME(EntityManager);
+
+			{
+				FTypeManager::FStaticTypeRegistrationHandle Handle = FTypeManager::RegisterStaticType<FTestSharedFragment_Int>();
+				// Handle goes out of scope and should deregister
+			}
+
+			// A new entity manager should NOT see the type
+			TSharedRef<FMassEntityManager> TestEntityManager = MakeInitializedEntityManager();
+			
+			const FTypeHandle TypeHandle = FTypeManager::MakeTypeHandle(FTestSharedFragment_Int::StaticStruct());
+			const FTypeInfo* TypeInfo = TestEntityManager->GetTypeManager().GetTypeInfo(TypeHandle);
+			AITEST_NULL("Deregistered type should not be in new TypeManager", TypeInfo);
+
+			TestEntityManager->Deinitialize();
+			return true;
+		}
+	};
+	IMPLEMENT_AI_INSTANT_TEST(FTypeManager_StaticRegistration_HandleDeregisters, "System.Mass.TypeManager.StaticRegistration.HandleDeregisters");
+
+	/** Verify that handle move semantics work: moved-from handle does not deregister. */
+	struct FTypeManager_StaticRegistration_HandleMoveSemantics : FExecutionTestBase
+	{
+		virtual bool InstantTest() override
+		{
+			CA_ASSUME(EntityManager);
+
+			FTypeManager::FStaticTypeRegistrationHandle TargetHandle;
+			{
+				FTypeManager::FStaticTypeRegistrationHandle SourceHandle = FTypeManager::RegisterStaticType<FTestSharedFragment_Int>();
+				AITEST_TRUE("Source handle is valid", SourceHandle.IsValid());
+
+				TargetHandle = MoveTemp(SourceHandle);
+				AITEST_FALSE("Moved-from handle is invalid", SourceHandle.IsValid());
+				AITEST_TRUE("Move target handle is valid", TargetHandle.IsValid());
+				// SourceHandle goes out of scope — should NOT deregister
+			}
+
+			// Type should still be in the static registry because TargetHandle holds ownership
+			TSharedRef<FMassEntityManager> TestEntityManager = MakeInitializedEntityManager();
+
+			const FTypeHandle TypeHandle = FTypeManager::MakeTypeHandle(FTestSharedFragment_Int::StaticStruct());
+			const FTypeInfo* TypeInfo = TestEntityManager->GetTypeManager().GetTypeInfo(TypeHandle);
+			AITEST_NOT_NULL("Type still registered after source handle destroyed", TypeInfo);
+
+			TestEntityManager->Deinitialize();
+
+			// Now destroy target handle — cleanup
+			TargetHandle = FTypeManager::FStaticTypeRegistrationHandle();
+			return true;
+		}
+	};
+	IMPLEMENT_AI_INSTANT_TEST(FTypeManager_StaticRegistration_HandleMoveSemantics, "System.Mass.TypeManager.StaticRegistration.HandleMoveSemantics");
+
+	//-------------------------------------------------------------------------
+	// Static registration automation tests
+	//-------------------------------------------------------------------------
+
+	/** Static registration survives across multiple EntityManager lifecycles. */
+	struct FTypeManager_StaticRegistration_MultiLifecycle : FExecutionTestBase
+	{
+		virtual bool InstantTest() override
+		{
+			FTypeManager::FStaticTypeRegistrationHandle Handle = FTypeManager::RegisterStaticType<UMassTestEngineSubsystem>();
+			const FTypeHandle TypeHandle = FTypeManager::MakeTypeHandle(UMassTestEngineSubsystem::StaticClass());
+
+			// First entity manager lifecycle
+			{
+				TSharedRef<FMassEntityManager> Manager1 = MakeInitializedEntityManager();
+				const FTypeInfo* TypeInfo = Manager1->GetTypeManager().GetTypeInfo(TypeHandle);
+				AITEST_NOT_NULL("Type present in first EntityManager", TypeInfo);
+				AITEST_NOT_NULL("First manager has subsystem traits", TypeInfo ? TypeInfo->GetAsSystemTraits() : nullptr);
+				Manager1->Deinitialize();
+			}
+
+			// Second entity manager lifecycle — static registration should still be there
+			{
+				TSharedRef<FMassEntityManager> Manager2 = MakeInitializedEntityManager();
+				const FTypeInfo* TypeInfo = Manager2->GetTypeManager().GetTypeInfo(TypeHandle);
+				AITEST_NOT_NULL("Type present in second EntityManager", TypeInfo);
+				Manager2->Deinitialize();
+			}
+
+			// Third — for good measure
+			{
+				TSharedRef<FMassEntityManager> Manager3 = MakeInitializedEntityManager();
+				const FTypeInfo* TypeInfo = Manager3->GetTypeManager().GetTypeInfo(TypeHandle);
+				AITEST_NOT_NULL("Type present in third EntityManager", TypeInfo);
+				Manager3->Deinitialize();
+			}
+
+			return true;
+		}
+	};
+	IMPLEMENT_AI_INSTANT_TEST(FTypeManager_StaticRegistration_MultiLifecycle, "System.Mass.TypeManager.StaticRegistration.MultiLifecycle");
+
+	/** Delegate handler overrides statically registered type (delegate runs after static drain). */
+	struct FTypeManager_StaticRegistration_DelegateOverridesStatic : FExecutionTestBase
+	{
+		FTypeManager::FStaticTypeRegistrationHandle StaticHandle;
+		FDelegateHandle BuiltInTypesDelegateHandle;
+
+		virtual bool SetUp() override
+		{
+			if (!FExecutionTestBase::SetUp())
+			{
+				return false;
+			}
+
+			StaticHandle = FTypeManager::RegisterStaticType<UMassTestWorldSubsystem>();
+			BuiltInTypesDelegateHandle = FTypeManager::OnRegisterBuiltInTypes.AddLambda([](FTypeManager& TypeMgr)
+			{
+				FSubsystemTypeTraits OverriddenTraits;
+				OverriddenTraits.bGameThreadOnly = false;
+				OverriddenTraits.bThreadSafeWrite = true;
+				TypeMgr.RegisterType(UMassTestWorldSubsystem::StaticClass(), MoveTemp(OverriddenTraits));
+			});
+			return true;
+		}
+
+		virtual void TearDown() override
+		{
+			FTypeManager::OnRegisterBuiltInTypes.Remove(BuiltInTypesDelegateHandle);
+			StaticHandle = FTypeManager::FStaticTypeRegistrationHandle();
+			FExecutionTestBase::TearDown();
+		}
+
+		virtual bool InstantTest() override
+		{
+			TSharedRef<FMassEntityManager> TestManager = MakeInitializedEntityManager();
+
+			const FTypeHandle TypeHandle = FTypeManager::MakeTypeHandle(UMassTestWorldSubsystem::StaticClass());
+			const FTypeInfo* TypeInfo = TestManager->GetTypeManager().GetTypeInfo(TypeHandle);
+			AITEST_NOT_NULL("Type is registered", TypeInfo);
+
+			const FSubsystemTypeTraits* Traits = TypeInfo ? TypeInfo->GetAsSystemTraits() : nullptr;
+			AITEST_NOT_NULL("Has subsystem traits", Traits);
+			if (Traits)
+			{
+				// Delegate runs AFTER static drain, so the delegate's traits should win
+				AITEST_FALSE("Delegate override: bGameThreadOnly should be false", Traits->bGameThreadOnly);
+				AITEST_TRUE("Delegate override: bThreadSafeWrite should be true", Traits->bThreadSafeWrite);
+			}
+
+			TestManager->Deinitialize();
+			return true;
+		}
+	};
+	IMPLEMENT_AI_INSTANT_TEST(FTypeManager_StaticRegistration_DelegateOverridesStatic, "System.Mass.TypeManager.StaticRegistration.DelegateOverridesStatic");
+
+	/** Instance-level RegisterType (from subsystem Initialize) and static registration coexist.
+	 * Per-instance registration happens before PostInitialize, static drain happens during PostInitialize.
+	 * Static drain overwrites the per-instance data (it runs in RegisterBuiltInTypes which is after Initialize). */
+	struct FTypeManager_StaticRegistration_OverridePrecedence : FExecutionTestBase
+	{
+		virtual bool InstantTest() override
+		{
+			// Register statically with specific traits
+			FSubsystemTypeTraits StaticTraits;
+			StaticTraits.bGameThreadOnly = false;
+			StaticTraits.bThreadSafeWrite = true;
+			FTypeManager::FStaticTypeRegistrationHandle Handle = FTypeManager::RegisterStaticType(
+				UMassTestEngineSubsystem::StaticClass(), MoveTemp(StaticTraits));
+
+			// Create a manager, register the same type with different traits BEFORE PostInitialize
+			TSharedRef<FMassEntityManager> TestManager = MakeInitializedEntityManager(/*bCallPostInitialize=*/false);
+
+			// Simulate what a subsystem's Initialize() would do — register before PostInitialize
+			FSubsystemTypeTraits InstanceTraits;
+			InstanceTraits.bGameThreadOnly = true;
+			InstanceTraits.bThreadSafeWrite = false;
+			TestManager->GetTypeManager().RegisterType(UMassTestEngineSubsystem::StaticClass(), MoveTemp(InstanceTraits));
+
+			// Now PostInitialize — this drains static types (which re-registers with static traits)
+			TestManager->PostInitialize();
+
+			const FTypeHandle TypeHandle = FTypeManager::MakeTypeHandle(UMassTestEngineSubsystem::StaticClass());
+			const FTypeInfo* TypeInfo = TestManager->GetTypeManager().GetTypeInfo(TypeHandle);
+			AITEST_NOT_NULL("Type is registered", TypeInfo);
+
+			const FSubsystemTypeTraits* Traits = TypeInfo ? TypeInfo->GetAsSystemTraits() : nullptr;
+			AITEST_NOT_NULL("Has subsystem traits", Traits);
+			if (Traits)
+			{
+				// Static drain runs during PostInitialize, AFTER instance-level registration,
+				// so the static traits overwrite the instance traits
+				AITEST_FALSE("Static override wins: bGameThreadOnly should be false", Traits->bGameThreadOnly);
+				AITEST_TRUE("Static override wins: bThreadSafeWrite should be true", Traits->bThreadSafeWrite);
+			}
+
+			TestManager->Deinitialize();
+			return true;
+		}
+	};
+	IMPLEMENT_AI_INSTANT_TEST(FTypeManager_StaticRegistration_OverridePrecedence, "System.Mass.TypeManager.StaticRegistration.OverridePrecedence");
+
+	/** Processor can declare a subsystem requirement on a statically registered type and resolve it. */
+	struct FTypeManager_StaticRegistration_ProcessorSubsystemRequirement : FExecutionTestBase
+	{
+		FTypeManager_StaticRegistration_ProcessorSubsystemRequirement()
+		{
+			bMakeWorldEntityManagersOwner = true;
+		}
+
+		virtual bool InstantTest() override
+		{
+			CA_ASSUME(EntityManager);
+
+			// The test's own EntityManager is already initialized.
+			// Register the engine subsystem type statically so it's known to the type system.
+			FTypeManager::FStaticTypeRegistrationHandle Handle = FTypeManager::RegisterStaticType<UMassTestEngineSubsystem>();
+
+			// Re-register into the existing EntityManager's TypeManager so it sees it too
+			// (the existing manager was already PostInitialized before our static registration)
+			EntityManager->GetTypeManager().RegisterType<UMassTestEngineSubsystem>();
+
+			// Set up a query with a subsystem requirement
+			FMassEntityQuery EntityQuery(EntityManager.ToSharedRef());
+			EntityQuery.AddRequirement<FTestFragment_Float>(EMassFragmentAccess::ReadOnly);
+			EntityQuery.AddSubsystemRequirement<UMassTestEngineSubsystem>(EMassFragmentAccess::ReadOnly);
+
+			// Create an entity so the query has something to iterate
+			const UScriptStruct* FragmentList[] = { FTestFragment_Float::StaticStruct() };
+			FMassArchetypeHandle Archetype = EntityManager->CreateArchetype(FragmentList);
+			EntityManager->CreateEntity(Archetype);
+
+			// Execute the query — verify the subsystem is accessible
+			UMassTestEngineSubsystem* ActualSubsystem = GEngine ? GEngine->GetEngineSubsystem<UMassTestEngineSubsystem>() : nullptr;
+			const UMassTestEngineSubsystem* AccessedSubsystem = nullptr;
+			bool bQueryExecuted = false;
+
+			FMassExecutionContext ExecutionContext(*EntityManager);
+			EntityQuery.ForEachEntityChunk(ExecutionContext, [&](FMassExecutionContext& Context)
+			{
+				bQueryExecuted = true;
+				AccessedSubsystem = Context.GetSubsystem<UMassTestEngineSubsystem>();
+			});
+
+			AITEST_TRUE("Query executed at least once", bQueryExecuted);
+			if (ActualSubsystem)
+			{
+				AITEST_EQUAL("Accessed subsystem matches actual engine subsystem", ActualSubsystem, AccessedSubsystem);
+			}
+
+			return true;
+		}
+	};
+	IMPLEMENT_AI_INSTANT_TEST(FTypeManager_StaticRegistration_ProcessorSubsystemRequirement, "System.Mass.TypeManager.StaticRegistration.ProcessorSubsystemRequirement");
+
+	/** Assigning a new registration to an existing handle deregisters the old type. */
+	struct FTypeManager_StaticRegistration_HandleReassignment : FExecutionTestBase
+	{
+		virtual bool InstantTest() override
+		{
+			// Register shared fragment first
+			FTypeManager::FStaticTypeRegistrationHandle Handle = FTypeManager::RegisterStaticType<FTestSharedFragment_Int>();
+			const FTypeHandle SharedFragHandle = FTypeManager::MakeTypeHandle(FTestSharedFragment_Int::StaticStruct());
+
+			// Now reassign the handle to hold a subsystem type instead
+			Handle = FTypeManager::RegisterStaticType<UMassTestEngineSubsystem>();
+			const FTypeHandle SubsystemHandle = FTypeManager::MakeTypeHandle(UMassTestEngineSubsystem::StaticClass());
+
+			// Create a new manager and verify: old type gone, new type present
+			TSharedRef<FMassEntityManager> TestManager = MakeInitializedEntityManager();
+
+			const FTypeInfo* SharedFragInfo = TestManager->GetTypeManager().GetTypeInfo(SharedFragHandle);
+			AITEST_NULL("Old type (shared fragment) should be deregistered", SharedFragInfo);
+
+			const FTypeInfo* SubsystemInfo = TestManager->GetTypeManager().GetTypeInfo(SubsystemHandle);
+			AITEST_NOT_NULL("New type (subsystem) should be registered", SubsystemInfo);
+			AITEST_NOT_NULL("New type has subsystem traits", SubsystemInfo ? SubsystemInfo->GetAsSystemTraits() : nullptr);
+
+			TestManager->Deinitialize();
+			return true;
+		}
+	};
+	IMPLEMENT_AI_INSTANT_TEST(FTypeManager_StaticRegistration_HandleReassignment, "System.Mass.TypeManager.StaticRegistration.HandleReassignment");
+} // namespace UE::Mass::Test
 
 UE_ENABLE_OPTIMIZATION_SHIP
 
diff --git a/Engine/Source/Runtime/MassEntity/Private/MassTypeManager.cpp b/Engine/Source/Runtime/MassEntity/Private/MassTypeManager.cpp
--- a/Engine/Source/Runtime/MassEntity/Private/MassTypeManager.cpp
+++ b/Engine/Source/Runtime/MassEntity/Private/MassTypeManager.cpp
@@ -6,10 +6,64 @@
 #include "Mass/EntityElementTypes.h"
 #include "MassEntityRelations.h"
 #include "Misc/CoreDelegates.h"
+#include "Misc/MTAccessDetector.h"
 
 
 namespace UE::Mass
 {
+	namespace Private
+	{
+		struct FStaticTypeEntry
+		{
+			FTypeInfo TypeInfo;
+			int32 RefCount = 0;
+		};
+
+		struct FStaticTypeRegistry
+		{
+			TMap<TObjectKey<const UStruct>, FStaticTypeEntry> Map;
+			UE_MT_DECLARE_RW_ACCESS_DETECTOR(MTAccessDetector);
+		};
+
+		/** Static registry of types that should be registered into every FTypeManager instance during RegisterBuiltInTypes(). */
+		FStaticTypeRegistry& GetStaticTypeRegistry()
+		{
+			static FStaticTypeRegistry Registry;
+			return Registry;
+		}
+
+		TObjectKey<const UStruct> AddToStaticRegistry(TNotNull<const UStruct*> InType, FTypeInfo&& TypeInfo)
+		{
+			TObjectKey<const UStruct> TypeKey(InType);
+
+			FStaticTypeRegistry& Registry = GetStaticTypeRegistry();
+			UE_MT_SCOPED_WRITE_ACCESS(Registry.MTAccessDetector);
+
+			FStaticTypeEntry* Existing = Registry.Map.Find(TypeKey);
+			if (Existing)
+			{
+				++Existing->RefCount;
+			}
+			else
+			{
+				Registry.Map.Add(TypeKey, FStaticTypeEntry{ MoveTemp(TypeInfo), 1 });
+			}
+			return TypeKey;
+		}
+
+		void UnregisterStaticTypeInternal(TObjectKey<const UStruct> TypeKey)
+		{
+			FStaticTypeRegistry& Registry = GetStaticTypeRegistry();
+			UE_MT_SCOPED_WRITE_ACCESS(Registry.MTAccessDetector);
+
+			FStaticTypeEntry* Entry = Registry.Map.Find(TypeKey);
+			if (Entry && --Entry->RefCount <= 0)
+			{
+				Registry.Map.Remove(TypeKey);
+			}
+		}
+	} // namespace Private
+
 	//-----------------------------------------------------------------------------
 	// FTypeHandle
 	//-----------------------------------------------------------------------------
@@ -31,6 +85,26 @@ namespace UE::Mass
 
 	void FTypeManager::RegisterBuiltInTypes()
 	{
+		if (!ensureMsgf(!bBuiltInTypesRegistered, TEXT("%hs: built-in types already registered!"), __FUNCTION__))
+		{
+			return;
+		}
+
+		// Drain statically registered types first, so delegate handlers can see/override them.
+		// We copy rather than move because the static registry is shared across all instances.
+		UE::Mass::Private::FStaticTypeRegistry& Registry = UE::Mass::Private::GetStaticTypeRegistry();
+		UE_MT_SCOPED_READ_ACCESS(Registry.MTAccessDetector);
+
+		for (TPair<TObjectKey<const UStruct>, UE::Mass::Private::FStaticTypeEntry>& Pair : Registry.Map)
+		{
+			const UStruct* ResolvedType = Pair.Key.ResolveObjectPtr();
+			if (ensureMsgf(ResolvedType, TEXT("Static type registration references a stale UStruct pointer")))
+			{
+				FTypeInfo TypeInfoCopy = Pair.Value.TypeInfo;
+				RegisterTypeInternal(ResolvedType, MoveTemp(TypeInfoCopy));
+			}
+		}
+
 		OnRegisterBuiltInTypes.Broadcast(*this);
 		bBuiltInTypesRegistered = true;
 	}
@@ -154,4 +228,67 @@ namespace UE::Mass
 	{
 		return UE::Mass::IsA<FMassRelation>(RelationOrElementType) && TypeDataMap.Contains(FTypeHandle(RelationOrElementType));
 	}
+
+	//-----------------------------------------------------------------------------
+	// FStaticTypeRegistrationHandle
+	//-----------------------------------------------------------------------------
+	FTypeManager::FStaticTypeRegistrationHandle::~FStaticTypeRegistrationHandle()
+	{
+		if (IsValid())
+		{
+			UE::Mass::Private::UnregisterStaticTypeInternal(this->RegisteredType);
+			this->RegisteredType = TObjectKey<const UStruct>();
+		}
+	}
+
+	FTypeManager::FStaticTypeRegistrationHandle::FStaticTypeRegistrationHandle(FStaticTypeRegistrationHandle&& InOther)
+		: RegisteredType(InOther.RegisteredType)
+	{
+		InOther.RegisteredType = TObjectKey<const UStruct>();
+	}
+
+	FTypeManager::FStaticTypeRegistrationHandle& FTypeManager::FStaticTypeRegistrationHandle::operator=(FStaticTypeRegistrationHandle&& InOther)
+	{
+		if (this != &InOther)
+		{
+			if (IsValid())
+			{
+				UE::Mass::Private::UnregisterStaticTypeInternal(this->RegisteredType);
+			}
+			this->RegisteredType = InOther.RegisteredType;
+			InOther.RegisteredType = TObjectKey<const UStruct>();
+		}
+		return *this;
+	}
+
+	//-----------------------------------------------------------------------------
+	// Static type registration (handle-based)
+	//-----------------------------------------------------------------------------
+	FTypeManager::FStaticTypeRegistrationHandle FTypeManager::RegisterStaticType(TNotNull<const UStruct*> InType, FSubsystemTypeTraits&& TypeTraits)
+	{
+		FTypeInfo TypeInfo;
+		TypeInfo.TypeName = InType->GetFName();
+		TypeInfo.Traits.Set<FSubsystemTypeTraits>(MoveTemp(TypeTraits));
+
+		FStaticTypeRegistrationHandle Handle;
+		Handle.RegisteredType = UE::Mass::Private::AddToStaticRegistry(InType, MoveTemp(TypeInfo));
+		return Handle;
+	}
+
+	FTypeManager::FStaticTypeRegistrationHandle FTypeManager::RegisterStaticType(TNotNull<const UStruct*> InType, FSharedFragmentTypeTraits&& TypeTraits)
+	{
+		if (!ensureMsgf(UE::Mass::IsA<FMassSharedFragment>(InType), TEXT("%s is not a valid shared fragment type"), *InType->GetName()))
+		{
+			return FStaticTypeRegistrationHandle();
+		}
+
+		FTypeInfo TypeInfo;
+		TypeInfo.TypeName = InType->GetFName();
+		TypeInfo.Traits.Set<FSharedFragmentTypeTraits>(MoveTemp(TypeTraits));
+
+		FStaticTypeRegistrationHandle Handle;
+		Handle.RegisteredType = UE::Mass::Private::AddToStaticRegistry(InType, MoveTemp(TypeInfo));
+		return Handle;
+	}
+
 } // namespace UE::Mass
diff --git a/Engine/Source/Runtime/MassEntity/Public/MassTypeManager.h b/Engine/Source/Runtime/MassEntity/Public/MassTypeManager.h
--- a/Engine/Source/Runtime/MassEntity/Public/MassTypeManager.h
+++ b/Engine/Source/Runtime/MassEntity/Public/MassTypeManager.h
@@ -141,6 +141,29 @@ namespace UE::Mass
 	{
 		DECLARE_MULTICAST_DELEGATE_OneParam(FOnRegisterBuiltInTypes, FTypeManager& /*Instance*/);
 
+		/**
+		 * Move-only, ref-counted handle for a static type registration. The static registry tracks
+		 * how many handles reference each type — the entry is only removed when the last handle
+		 * for that type is destroyed. Already-initialized FTypeManager instances are unaffected.
+		 */
+		struct FStaticTypeRegistrationHandle
+		{
+			FStaticTypeRegistrationHandle() = default;
+			UE_API ~FStaticTypeRegistrationHandle();
+
+			UE_API FStaticTypeRegistrationHandle(FStaticTypeRegistrationHandle&& InOther);
+			UE_API FStaticTypeRegistrationHandle& operator=(FStaticTypeRegistrationHandle&& InOther);
+
+			FStaticTypeRegistrationHandle(const FStaticTypeRegistrationHandle&) = delete;
+			FStaticTypeRegistrationHandle& operator=(const FStaticTypeRegistrationHandle&) = delete;
+
+			bool IsValid() const;
+
+		private:
+			friend struct FTypeManager;
+			TObjectKey<const UStruct> RegisteredType;
+		};
+
 		UE_API explicit FTypeManager(FMassEntityManager& InEntityManager);
 
 		void RegisterBuiltInTypes();
@@ -203,6 +226,28 @@ namespace UE::Mass
 		 */
 		UE_API static FOnRegisterBuiltInTypes OnRegisterBuiltInTypes;
 
+		//---------------------------------------------------------------------
+		// Static type registration
+		// Types registered here are copied into every FTypeManager instance
+		// during RegisterBuiltInTypes(). Use for types whose lifetime spans
+		// multiple entity managers (engine subsystems, game instance subsystems).
+		//---------------------------------------------------------------------
+
+		/**
+		 * Register a type statically with an RAII handle. Destroying the returned handle removes
+		 * the type from the static registry (future TypeManager instances won't see it).
+		 */
+		UE_API static FStaticTypeRegistrationHandle RegisterStaticType(TNotNull<const UStruct*> InType, FSubsystemTypeTraits&& TypeTraits);
+		UE_API static FStaticTypeRegistrationHandle RegisterStaticType(TNotNull<const UStruct*> InType, FSharedFragmentTypeTraits&& TypeTraits);
+
+		/** RAII static registration for subsystem types. T must derive from USubsystem. */
+		template<CSubsystem T> 
+		static FStaticTypeRegistrationHandle RegisterStaticType();
+
+		/** RAII static registration for shared fragment types. T must satisfy CSharedFragment. */
+		template<CSharedFragment T>
+		static FStaticTypeRegistrationHandle RegisterStaticType();
+
 	private:
 		/** Register traits for given type */
 		UE_API FTypeHandle RegisterTypeInternal(TNotNull<const UStruct*> InType, FTypeInfo&&);
@@ -304,10 +349,31 @@ namespace UE::Mass
 		return OuterEntityManager;
 	}
 
+	inline bool FTypeManager::FStaticTypeRegistrationHandle::IsValid() const
+	{
+		return RegisteredType != TObjectKey<const UStruct>();
+	}
+
 	inline FTypeHandle FTypeManager::MakeTypeHandle(TNotNull<const UStruct*> InTypeKey)
 	{
 		return FTypeHandle(InTypeKey);
 	}
+
+	//---------------------------------------------------------------------
+	// Static registration template inlines
+	//---------------------------------------------------------------------
+	template<CSubsystem T>
+	FTypeManager::FStaticTypeRegistrationHandle FTypeManager::RegisterStaticType()
+	{
+		return RegisterStaticType(T::StaticClass(), FSubsystemTypeTraits::Make<T>());
+	}
+
+	template<CSharedFragment T>
+	FTypeManager::FStaticTypeRegistrationHandle FTypeManager::RegisterStaticType()
+	{
+		return RegisterStaticType(T::StaticStruct(), FSharedFragmentTypeTraits::Make<T>());
+	}
+
 } // namespace UE::Mass
 
 #undef UE_API
