[MassEntity] Fix missing observer notification in BatchAddSharedFragmentsForEntities
Add tests verifying that BatchAddSharedFragmentsForEntities fires AddElement observers for both shared and const shared fragment types

#Synced-CL 52643246

diff --git a/Engine/Source/Developer/MassEntityTestSuite/Private/MassObserverTest.cpp b/Engine/Source/Developer/MassEntityTestSuite/Private/MassObserverTest.cpp
--- a/Engine/Source/Developer/MassEntityTestSuite/Private/MassObserverTest.cpp
+++ b/Engine/Source/Developer/MassEntityTestSuite/Private/MassObserverTest.cpp
@@ -4,6 +4,7 @@
 #include "MassProcessingTypes.h"
 #include "MassEntityTestTypes.h"
 #include "MassEntityTypes.h"
+#include "MassEntityUtils.h"
 #include "MassCommandBuffer.h"
 #include "MassExecutionContext.h"
 #include "MassObserverNotificationTypes.h"
@@ -1220,6 +1221,84 @@ struct FCompositeObserver_GetObservedTypes : FEntityTestBase
 };
 IMPLEMENT_AI_INSTANT_TEST(FCompositeObserver_GetObservedTypes, "System.Mass.Observer.CompositeObserver.GetObservedTypes");
 
+/**
+ * Verifies that BatchAddSharedFragmentsForEntities fires the AddElement observer for the shared fragment type.
+ */
+struct FBatchAddSharedFragment_ObserverFires : FEntityTestBase
+{
+	virtual bool InstantTest() override
+	{
+		TArray<FMassEntityHandle> AffectedEntities;
+		UMassTestProcessorBase* ObserverProcessor = NewTestProcessor<UMassTestProcessorBase>(EntityManager);
+		ObserverProcessor->EntityQuery.AddRequirement<FTestFragment_Int>(EMassFragmentAccess::ReadOnly);
+		ObserverProcessor->EntityQuery.AddSharedRequirement<FTestSharedFragment_Int>(EMassFragmentAccess::ReadOnly);
+		ObserverProcessor->ForEachEntityChunkExecutionFunction = [&AffectedEntities](FMassExecutionContext& Context)
+		{
+			AffectedEntities.Append(Context.GetEntities().GetData(), Context.GetEntities().Num());
+		};
+
+		FMassObserverManager& ObserverManager = EntityManager->GetObserverManager();
+		ObserverManager.AddObserverInstance(FTestSharedFragment_Int::StaticStruct(), EMassObservedOperationFlags::Add, ObserverProcessor);
+
+		TArray<FMassEntityHandle> Entities;
+		EntityManager->BatchCreateEntities(IntsArchetype, 5, Entities);
+
+		FTestSharedFragment_Int FragmentInstance;
+		FragmentInstance.Value = 42;
+		const FSharedStruct SharedFragmentInstance = EntityManager->GetOrCreateSharedFragment(FragmentInstance);
+		FMassArchetypeSharedFragmentValues SharedValues;
+		SharedValues.Add(SharedFragmentInstance);
+
+		TArray<FMassArchetypeEntityCollection> EntityCollections;
+		UE::Mass::Utils::CreateEntityCollections(*EntityManager, Entities, FMassArchetypeEntityCollection::EDuplicatesHandling::NoDuplicates, EntityCollections);
+		EntityManager->BatchAddSharedFragmentsForEntities(EntityCollections, SharedValues);
+
+		AITEST_EQUAL(TEXT("Observer should fire for all entities that received the shared fragment"), AffectedEntities.Num(), Entities.Num());
+
+		return true;
+	}
+};
+IMPLEMENT_AI_INSTANT_TEST(FBatchAddSharedFragment_ObserverFires, "System.Mass.Observer.SharedFragment.BatchAddFiresObserver");
+
+/**
+ * Verifies that BatchAddSharedFragmentsForEntities fires the AddElement observer for a const shared fragment type.
+ */
+struct FBatchAddConstSharedFragment_ObserverFires : FEntityTestBase
+{
+	virtual bool InstantTest() override
+	{
+		TArray<FMassEntityHandle> AffectedEntities;
+		UMassTestProcessorBase* ObserverProcessor = NewTestProcessor<UMassTestProcessorBase>(EntityManager);
+		ObserverProcessor->EntityQuery.AddRequirement<FTestFragment_Int>(EMassFragmentAccess::ReadOnly);
+		ObserverProcessor->EntityQuery.AddConstSharedRequirement<FTestConstSharedFragment_Int>();
+		ObserverProcessor->ForEachEntityChunkExecutionFunction = [&AffectedEntities](FMassExecutionContext& Context)
+		{
+			AffectedEntities.Append(Context.GetEntities().GetData(), Context.GetEntities().Num());
+		};
+
+		FMassObserverManager& ObserverManager = EntityManager->GetObserverManager();
+		ObserverManager.AddObserverInstance(FTestConstSharedFragment_Int::StaticStruct(), EMassObservedOperationFlags::Add, ObserverProcessor);
+
+		TArray<FMassEntityHandle> Entities;
+		EntityManager->BatchCreateEntities(IntsArchetype, 5, Entities);
+
+		FTestConstSharedFragment_Int FragmentInstance;
+		FragmentInstance.Value = 99;
+		const FConstSharedStruct ConstSharedFragmentInstance = EntityManager->GetOrCreateConstSharedFragment(FragmentInstance);
+		FMassArchetypeSharedFragmentValues SharedValues;
+		SharedValues.Add(ConstSharedFragmentInstance);
+
+		TArray<FMassArchetypeEntityCollection> EntityCollections;
+		UE::Mass::Utils::CreateEntityCollections(*EntityManager, Entities, FMassArchetypeEntityCollection::EDuplicatesHandling::NoDuplicates, EntityCollections);
+		EntityManager->BatchAddSharedFragmentsForEntities(EntityCollections, SharedValues);
+
+		AITEST_EQUAL(TEXT("Observer should fire for all entities that received the const shared fragment"), AffectedEntities.Num(), Entities.Num());
+
+		return true;
+	}
+};
+IMPLEMENT_AI_INSTANT_TEST(FBatchAddConstSharedFragment_ObserverFires, "System.Mass.Observer.SharedFragment.BatchAddConstSharedFiresObserver");
+
 } // UE::Mass::Test::Observers
 
 UE_ENABLE_OPTIMIZATION_SHIP
diff --git a/Engine/Source/Runtime/MassEntity/Private/MassEntityManager.cpp b/Engine/Source/Runtime/MassEntity/Private/MassEntityManager.cpp
--- a/Engine/Source/Runtime/MassEntity/Private/MassEntityManager.cpp
+++ b/Engine/Source/Runtime/MassEntity/Private/MassEntityManager.cpp
@@ -2054,13 +2054,16 @@ void FMassEntityManager::BatchAddSharedFragmentsForEntities(TConstArrayView<FMas
 
 	TRACE_CPUPROFILER_EVENT_SCOPE(Mass_BatchAddConstSharedFragmentForEntities);
 
+	const FMassElementBitSet ElementsToAddBitSet = AddedFragmentValues.GetBitSet();
+	const bool bElementsAddedAreObserved = ObserverManager.HasObserversForBitSet(ElementsToAddBitSet, EMassObservedOperation::AddElement);
+
 	for (const FMassArchetypeEntityCollection& Collection : EntityCollections)
 	{
 		FMassArchetypeData* CurrentArchetype = Collection.GetArchetype().DataPtr.Get();
 		testableCheckfReturn(CurrentArchetype, continue, TEXT("Adding shared fragments to archetype-less entities is not supported"));
 
-		FMassArchetypeCompositionDescriptor NewComposition(CurrentArchetype->GetCompositionDescriptor());
-		NewComposition.GetElementsBitSet() += AddedFragmentValues.GetBitSet();
+		FMassElementBitSet NewComposition(CurrentArchetype->GetCompositionBitSet());
+		NewComposition += ElementsToAddBitSet;
 
 		const FMassArchetypeHandle NewArchetypeHandle = CreateArchetype(NewComposition, FMassArchetypeCreationParams(*CurrentArchetype));
 		check(NewArchetypeHandle.IsValid());
@@ -2072,15 +2075,36 @@ void FMassEntityManager::BatchAddSharedFragmentsForEntities(TConstArrayView<FMas
 			continue;
 		}
 
+		// Compute BitSet of the elements actually added (i.e., not already present)
+		FMassElementBitSet ElementsAdded;
+		if (bElementsAddedAreObserved)
+		{
+			ElementsAdded = ElementsToAddBitSet - CurrentArchetype->GetCompositionBitSet();
+		}
+
 		TArray<FMassEntityHandle> EntitiesBeingMoved;
-		CurrentArchetype->BatchMoveEntitiesToAnotherArchetype(Collection, *NewArchetype, EntitiesBeingMoved, /*OutNewChunks=*/nullptr, &AddedFragmentValues);
+		FMassArchetypeEntityCollection::FEntityRangeArray NewArchetypeEntityRanges;
+
+		CurrentArchetype->BatchMoveEntitiesToAnotherArchetype(Collection
+			, *NewArchetype
+			, EntitiesBeingMoved
+			, !ElementsAdded.IsEmpty() ? &NewArchetypeEntityRanges : nullptr
+			, &AddedFragmentValues);
 
 		for (const FMassEntityHandle& Entity : EntitiesBeingMoved)
 		{
 			check(GetEntityStorageInterface().IsValidIndex(Entity.Index));
-			
+
 			GetEntityStorageInterface().SetArchetypeFromShared(Entity.Index, NewArchetypeHandle.DataPtr);
 		}
+
+		if (!NewArchetypeEntityRanges.IsEmpty())
+		{
+			ObserverManager.OnCompositionChanged(
+				FMassArchetypeEntityCollection(NewArchetypeHandle, MoveTemp(NewArchetypeEntityRanges))
+				, MoveTemp(ElementsAdded)
+				, EMassObservedOperation::AddElement);
+		}
 	}
 }
 
