UE: IWYU fix

diff --git a/Engine/Source/Developer/MassEntityTestSuite/Public/MassEntityTestFarmPlot.h b/Engine/Source/Developer/MassEntityTestSuite/Public/MassEntityTestFarmPlot.h
--- a/Engine/Source/Developer/MassEntityTestSuite/Public/MassEntityTestFarmPlot.h
+++ b/Engine/Source/Developer/MassEntityTestSuite/Public/MassEntityTestFarmPlot.h
@@ -9,6 +9,7 @@
 #include "Engine/DataTable.h"
 #include "MassEntityTypes.h"
 #include "MassEntityManager.h"
+#include "MassEntityQuery.h"
 #include "MassCommandBuffer.h"
 #include "MassProcessor.h"
 #if UE_ENABLE_INCLUDE_ORDER_DEPRECATED_IN_5_6
